Sequencer: Time range slider

- An optional slider at the bottom of the timeline that show a start/end and in/out range. When enabled, adjusting the timeline (panning/zooming) will be locked to the start/end range. This is another option which is on  by default.
- Added project settings for the time range. For the level editor sequencer, it defaults to frames 101-200.
- When the slider is not visible (the default for UMG), there should be no change in behavior.

diff --git a/Engine/Source/Developer/LogVisualizer/Private/STimelinesContainer.cpp b/Engine/Source/Developer/LogVisualizer/Private/STimelinesContainer.cpp
--- a/Engine/Source/Developer/LogVisualizer/Private/STimelinesContainer.cpp
+++ b/Engine/Source/Developer/LogVisualizer/Private/STimelinesContainer.cpp
@@ -320,15 +320,13 @@ void STimelinesContainer::OnNewLogEntry(const FVisualLogDevice::FVisualLogEntryI
 	CachedMaxTime = CachedMaxTime > Entry.Entry.TimeStamp ? CachedMaxTime : Entry.Entry.TimeStamp;
 
 	TRange<float> LocalViewRange = TimeSliderController->GetTimeSliderArgs().ViewRange.Get();
-	const float CurrentMin = TimeSliderController->GetTimeSliderArgs().ClampMin.Get().GetValue();
-	const float CurrentMax = TimeSliderController->GetTimeSliderArgs().ClampMax.Get().GetValue();
-	float ZoomLevel = LocalViewRange.Size<float>() / (CurrentMax - CurrentMin);
+	TRange<float> ClampRange = TimeSliderController->GetTimeSliderArgs().ClampRange.Get();
+	float ZoomLevel = LocalViewRange.Size<float>() / ClampRange.Size<float>();
 
-	TimeSliderController->GetTimeSliderArgs().ClampMin = CachedMinTime;
-	TimeSliderController->GetTimeSliderArgs().ClampMax = CachedMaxTime + 0.1;
+	TimeSliderController->GetTimeSliderArgs().ClampRange = TRange<float>(CachedMinTime, CachedMaxTime + 0.1f);
 	if ( FMath::Abs(ZoomLevel - 1) <= SMALL_NUMBER)
 	{
-		TimeSliderController->GetTimeSliderArgs().ViewRange = TRange<float>(TimeSliderController->GetTimeSliderArgs().ClampMin.Get().GetValue(), TimeSliderController->GetTimeSliderArgs().ClampMax.Get().GetValue());
+		TimeSliderController->GetTimeSliderArgs().ViewRange = TimeSliderController->GetTimeSliderArgs().ClampRange;
 	}
 }
 
diff --git a/Engine/Source/Developer/LogVisualizer/Private/SVisualLoggerView.cpp b/Engine/Source/Developer/LogVisualizer/Private/SVisualLoggerView.cpp
--- a/Engine/Source/Developer/LogVisualizer/Private/SVisualLoggerView.cpp
+++ b/Engine/Source/Developer/LogVisualizer/Private/SVisualLoggerView.cpp
@@ -76,8 +76,7 @@ void SVisualLoggerView::Construct(const FArguments& InArgs, const TSharedRef<FUI
 
 	FVisualLoggerTimeSliderArgs TimeSliderArgs;
 	TimeSliderArgs.ViewRange = FAnimatedRange::WrapAttribute(InArgs._ViewRange);
-	TimeSliderArgs.ClampMin = InArgs._ViewRange.Get().GetLowerBoundValue();
-	TimeSliderArgs.ClampMax = InArgs._ViewRange.Get().GetUpperBoundValue();
+	TimeSliderArgs.ClampRange = TRange<float>(InArgs._ViewRange.Get().GetLowerBoundValue(), InArgs._ViewRange.Get().GetUpperBoundValue());
 	TimeSliderArgs.ScrubPosition = InArgs._ScrubPosition;
 	FLogVisualizer::Get().GetTimeSliderController()->SetTimesliderArgs(TimeSliderArgs);
 
diff --git a/Engine/Source/Developer/LogVisualizer/Private/TimeSliderController.cpp b/Engine/Source/Developer/LogVisualizer/Private/TimeSliderController.cpp
--- a/Engine/Source/Developer/LogVisualizer/Private/TimeSliderController.cpp
+++ b/Engine/Source/Developer/LogVisualizer/Private/TimeSliderController.cpp
@@ -408,25 +408,25 @@ FReply FVisualLoggerTimeSliderController::OnMouseMove( TSharedRef<SWidget> Widge
 				float NewViewOutputMin = LocalViewRangeMin - InputDelta.X;
 				float NewViewOutputMax = LocalViewRangeMax - InputDelta.X;
 
-				TOptional<float> LocalClampMin = TimeSliderArgs.ClampMin.Get();
-				TOptional<float> LocalClampMax = TimeSliderArgs.ClampMax.Get();
+				float LocalClampMin = TimeSliderArgs.ClampRange.Get().GetLowerBoundValue();
+				float LocalClampMax = TimeSliderArgs.ClampRange.Get().GetUpperBoundValue();
 
-				// Clamp the range if clamp values are set
-				if ( LocalClampMin.IsSet() && NewViewOutputMin < LocalClampMin.GetValue() )
+				// Clamp the range
+				if ( NewViewOutputMin < LocalClampMin )
 				{
-					NewViewOutputMin = LocalClampMin.GetValue();
+					NewViewOutputMin = LocalClampMin;
 				}
 			
-				if ( LocalClampMax.IsSet() && NewViewOutputMax > LocalClampMax.GetValue() )
+				if ( NewViewOutputMax > LocalClampMax )
 				{
-					NewViewOutputMax = LocalClampMax.GetValue();
+					NewViewOutputMax = LocalClampMax;
 				}
 
 				TimeSliderArgs.OnViewRangeChanged.ExecuteIfBound(TRange<float>(NewViewOutputMin, NewViewOutputMax), EViewRangeInterpolation::Immediate);
 				if (Scrollbar.IsValid())
 				{
-					float InOffsetFraction = (NewViewOutputMin - LocalClampMin.GetValue()) / (LocalClampMax.GetValue() - LocalClampMin.GetValue());
-					float InThumbSizeFraction = (NewViewOutputMax - NewViewOutputMin) / (LocalClampMax.GetValue() - LocalClampMin.GetValue());
+					float InOffsetFraction = (NewViewOutputMin - LocalClampMin) / (LocalClampMax - LocalClampMin);
+					float InThumbSizeFraction = (NewViewOutputMax - NewViewOutputMin) / (LocalClampMax - LocalClampMin);
 					Scrollbar->SetState(InOffsetFraction, InThumbSizeFraction);
 				}
 
@@ -453,14 +453,18 @@ FReply FVisualLoggerTimeSliderController::OnMouseMove( TSharedRef<SWidget> Widge
 				FScrubRangeToScreen RangeToScreen( TimeSliderArgs.ViewRange.Get(), MyGeometry.Size );
 				FVector2D CursorPos = MyGeometry.AbsoluteToLocal( MouseEvent.GetLastScreenSpacePosition() );
 				float NewValue = RangeToScreen.LocalXToInput( CursorPos.X );
-				if (TimeSliderArgs.ClampMin.Get().IsSet() && NewValue < TimeSliderArgs.ClampMin.Get().GetValue())
+
+				float LocalClampMin = TimeSliderArgs.ClampRange.Get().GetLowerBoundValue();
+				float LocalClampMax = TimeSliderArgs.ClampRange.Get().GetUpperBoundValue();
+
+				if (NewValue < LocalClampMin)
 				{
-					NewValue = TimeSliderArgs.ClampMin.Get().GetValue();
+					NewValue = LocalClampMin;
 				}
 
-				if (TimeSliderArgs.ClampMax.Get().IsSet() && NewValue > TimeSliderArgs.ClampMax.Get().GetValue())
+				if (NewValue > LocalClampMax)
 				{
-					NewValue = TimeSliderArgs.ClampMax.Get().GetValue();
+					NewValue = LocalClampMax;
 				}
 
 				CommitScrubPosition(NewValue, /*bIsScrubbing=*/true);
@@ -507,17 +511,17 @@ void FVisualLoggerTimeSliderController::HorizontalScrollBar_OnUserScrolled(float
 		TRange<float> LocalViewRange = TimeSliderArgs.ViewRange.Get();
 		float LocalViewRangeMin = LocalViewRange.GetLowerBoundValue();
 		float LocalViewRangeMax = LocalViewRange.GetUpperBoundValue();
-		TOptional<float> LocalClampMin = TimeSliderArgs.ClampMin.Get();
-		TOptional<float> LocalClampMax = TimeSliderArgs.ClampMax.Get();
+		float LocalClampMin = TimeSliderArgs.ClampRange.Get().GetLowerBoundValue();
+		float LocalClampMax = TimeSliderArgs.ClampRange.Get().GetUpperBoundValue();
 
-		float InThumbSizeFraction = (LocalViewRangeMax - LocalViewRangeMin) / (LocalClampMax.GetValue() - LocalClampMin.GetValue());
+		float InThumbSizeFraction = (LocalViewRangeMax - LocalViewRangeMin) / (LocalClampMax - LocalClampMin);
 
-		float NewViewOutputMin = LocalClampMin.GetValue() + ScrollOffset * (LocalClampMax.GetValue() - LocalClampMin.GetValue());
+		float NewViewOutputMin = LocalClampMin + ScrollOffset * (LocalClampMax - LocalClampMin);
 		// The  output is not bound to a delegate so we'll manage the value ourselves
-		float NewViewOutputMax = FMath::Min<float>(NewViewOutputMin + (LocalViewRangeMax - LocalViewRangeMin), LocalClampMax.GetValue());
+		float NewViewOutputMax = FMath::Min<float>(NewViewOutputMin + (LocalViewRangeMax - LocalViewRangeMin), LocalClampMax);
 		NewViewOutputMin = NewViewOutputMax - (LocalViewRangeMax - LocalViewRangeMin);
 
-		float InOffsetFraction = (NewViewOutputMin - LocalClampMin.GetValue()) / (LocalClampMax.GetValue() - LocalClampMin.GetValue());
+		float InOffsetFraction = (NewViewOutputMin - LocalClampMin) / (LocalClampMax - LocalClampMin);
 		//if (InOffsetFraction + InThumbSizeFraction <= 1)
 		{
 			TimeSliderArgs.ViewRange.Set(TRange<float>(NewViewOutputMin, NewViewOutputMax));
@@ -530,26 +534,25 @@ void FVisualLoggerTimeSliderController::SetTimeRange(float NewViewOutputMin, flo
 {
 	TimeSliderArgs.ViewRange.Set(TRange<float>(NewViewOutputMin, NewViewOutputMax));
 
-	TOptional<float> LocalClampMin = TimeSliderArgs.ClampMin.Get();
-	TOptional<float> LocalClampMax = TimeSliderArgs.ClampMax.Get();
+	float LocalClampMin = TimeSliderArgs.ClampRange.Get().GetLowerBoundValue();
+	float LocalClampMax = TimeSliderArgs.ClampRange.Get().GetUpperBoundValue();
 
-	const float InOffsetFraction = (NewViewOutputMin - LocalClampMin.GetValue()) / (LocalClampMax.GetValue() - LocalClampMin.GetValue());
-	const float InThumbSizeFraction = (NewViewOutputMax - NewViewOutputMin) / (LocalClampMax.GetValue() - LocalClampMin.GetValue());
+	const float InOffsetFraction = (NewViewOutputMin - LocalClampMin) / (LocalClampMax - LocalClampMin);
+	const float InThumbSizeFraction = (NewViewOutputMax - NewViewOutputMin) / (LocalClampMax - LocalClampMin);
 	Scrollbar->SetState(InOffsetFraction, InThumbSizeFraction);
 }
 
 void FVisualLoggerTimeSliderController::SetClampRange(float MinValue, float MaxValue)
 {
 	TRange<float> LocalViewRange = TimeSliderArgs.ViewRange.Get();
-	float LocalClampMin = TimeSliderArgs.ClampMin.Get().GetValue();
-	float LocalClampMax = TimeSliderArgs.ClampMax.Get().GetValue();
+	float LocalClampMin = TimeSliderArgs.ClampRange.Get().GetLowerBoundValue();
+	float LocalClampMax = TimeSliderArgs.ClampRange.Get().GetUpperBoundValue();
 	const float CurrentDistance = LocalClampMax - LocalClampMin;
 	const float ZoomDelta = (LocalViewRange.GetUpperBoundValue() - LocalViewRange.GetLowerBoundValue()) / CurrentDistance;
 
 	MaxValue = MinValue + (MaxValue - MinValue < 2 ? CurrentDistance : MaxValue - MinValue);
 
-	TimeSliderArgs.ClampMin.Set(MinValue);
-	TimeSliderArgs.ClampMax.Set(MaxValue);
+	TimeSliderArgs.ClampRange = TRange<float>(MinValue, MaxValue);
 
 	const float LocalViewRangeMin = FMath::Clamp(LocalViewRange.GetLowerBoundValue(), MinValue, MaxValue);
 	const float LocalViewRangeMax = FMath::Clamp(LocalViewRange.GetUpperBoundValue(), MinValue, MaxValue);
@@ -583,25 +586,25 @@ FReply FVisualLoggerTimeSliderController::OnMouseWheel( TSharedRef<SWidget> Widg
 
 			if( FMath::Abs( OutputChange ) > 0.01f && NewViewOutputMin < NewViewOutputMax )
 			{
-				TOptional<float> LocalClampMin = TimeSliderArgs.ClampMin.Get();
-				TOptional<float> LocalClampMax = TimeSliderArgs.ClampMax.Get();
+				float LocalClampMin = TimeSliderArgs.ClampRange.Get().GetLowerBoundValue();
+				float LocalClampMax = TimeSliderArgs.ClampRange.Get().GetUpperBoundValue();
 
 				// Clamp the range if clamp values are set
-				if ( LocalClampMin.IsSet() && NewViewOutputMin < LocalClampMin.GetValue() )
+				if ( NewViewOutputMin < LocalClampMin )
 				{
-					NewViewOutputMin = LocalClampMin.GetValue();
+					NewViewOutputMin = LocalClampMin;
 				}
 				
-				if ( LocalClampMax.IsSet() && NewViewOutputMax > LocalClampMax.GetValue() )
+				if ( NewViewOutputMax > LocalClampMax )
 				{
-					NewViewOutputMax = LocalClampMax.GetValue();
+					NewViewOutputMax = LocalClampMax;
 				}
 
 				TimeSliderArgs.OnViewRangeChanged.ExecuteIfBound(TRange<float>(NewViewOutputMin, NewViewOutputMax), EViewRangeInterpolation::Immediate);
 				if (Scrollbar.IsValid())
 				{
-					float InOffsetFraction = (NewViewOutputMin - LocalClampMin.GetValue()) / (LocalClampMax.GetValue() - LocalClampMin.GetValue());
-					float InThumbSizeFraction = (NewViewOutputMax - NewViewOutputMin) / (LocalClampMax.GetValue() - LocalClampMin.GetValue());
+					float InOffsetFraction = (NewViewOutputMin - LocalClampMin) / (LocalClampMax - LocalClampMin);
+					float InThumbSizeFraction = (NewViewOutputMax - NewViewOutputMin) / (LocalClampMax - LocalClampMin);
 					Scrollbar->SetState(InOffsetFraction, InThumbSizeFraction);
 				}
 				if( !TimeSliderArgs.ViewRange.IsBound() )
diff --git a/Engine/Source/Editor/EditorStyle/Private/SlateEditorStyle.cpp b/Engine/Source/Editor/EditorStyle/Private/SlateEditorStyle.cpp
--- a/Engine/Source/Editor/EditorStyle/Private/SlateEditorStyle.cpp
+++ b/Engine/Source/Editor/EditorStyle/Private/SlateEditorStyle.cpp
@@ -1579,6 +1579,9 @@ void FSlateEditorStyle::FStyle::SetupGeneralStyles()
 		Set( "Sequencer.Timeline.ScrubHandleDown", new BOX_BRUSH( "Sequencer/ScrubHandleDown", FMargin( 6.f/13.f, 5/12.f, 6/13.f, 8/12.f ) ) );
 		Set( "Sequencer.Timeline.ScrubHandleUp", new BOX_BRUSH( "Sequencer/ScrubHandleUp", FMargin( 6.f/13.f, 8/12.f, 6/13.f, 5/12.f ) ) );
 		Set( "Sequencer.Timeline.ScrubHandleWhole", new BOX_BRUSH( "Sequencer/ScrubHandleWhole", FMargin( 6.f/13.f, 10/24.f, 6/13.f, 10/24.f  ) ) );
+		Set( "Sequencer.Timeline.RangeHandleLeft", new BOX_BRUSH( "Sequencer/SectionGripLeft", FMargin(5.f/16.f) ) );
+		Set( "Sequencer.Timeline.RangeHandleRight", new BOX_BRUSH( "Sequencer/SectionGripRight", FMargin(5.f/16.f) ) );
+		Set( "Sequencer.Timeline.RangeHandle", new BOX_BRUSH( "Sequencer/GenericSectionBackground", FMargin(5.f/16.f) ) );
 		Set( "Sequencer.Timeline.NotifyAlignmentMarker", new IMAGE_BRUSH( "Sequencer/NotifyAlignmentMarker", FVector2D(10,19) ) );
 		Set( "Sequencer.NotificationImage_AddedPlayMovieSceneEvent", new IMAGE_BRUSH( "Old/Checkbox_checked", Icon16x16 ) );
 		Set( "Sequencer.Key", new IMAGE_BRUSH( "Common/Key", FVector2D(7.0f,12.0f) ) );
diff --git a/Engine/Source/Editor/NiagaraEditor/Private/NiagaraEffectEditor.cpp b/Engine/Source/Editor/NiagaraEditor/Private/NiagaraEffectEditor.cpp
--- a/Engine/Source/Editor/NiagaraEditor/Private/NiagaraEffectEditor.cpp
+++ b/Engine/Source/Editor/NiagaraEditor/Private/NiagaraEffectEditor.cpp
@@ -80,14 +80,21 @@ void FNiagaraEffectEditor::InitNiagaraEffectEditor(const EToolkitMode::Type Mode
 	check(Effect != NULL);
 	EffectInstance = new FNiagaraEffectInstance(InEffect);
 
+	const float InTime = -0.02f;
+	const float OutTime = 3.2f;
+
 	if (!Sequencer.IsValid())
 	{
 		MovieScene = NewObject<UMovieScene>(InEffect, FName("Niagara Effect MovieScene"), RF_RootSet);
 		auto NewAnimation = NewObject<UNiagaraAnimation>(MovieScene);
+		MovieScene->StartTime = InTime;
+		MovieScene->InTime = InTime;
+		MovieScene->OutTime = OutTime;
+		MovieScene->EndTime = OutTime;
 
 		FSequencerViewParams ViewParams(TEXT("NiagaraSequencerSettings"));
 		{
-			ViewParams.InitalViewRange = TRange<float>(-0.02f, 3.2f);
+			ViewParams.InitalViewRange = TRange<float>(InTime, OutTime);
 			ViewParams.InitialScrubPosition = 0;
 		}
 
diff --git a/Engine/Source/Editor/Sequencer/Private/SSequencer.cpp b/Engine/Source/Editor/Sequencer/Private/SSequencer.cpp
--- a/Engine/Source/Editor/Sequencer/Private/SSequencer.cpp
+++ b/Engine/Source/Editor/Sequencer/Private/SSequencer.cpp
@@ -193,7 +193,9 @@ void SSequencer::Construct( const FArguments& InArgs, TSharedRef< class FSequenc
 
 	FTimeSliderArgs TimeSliderArgs;
 	TimeSliderArgs.ViewRange = InArgs._ViewRange;
+	TimeSliderArgs.ClampRange = InArgs._ClampRange;
 	TimeSliderArgs.OnViewRangeChanged = InArgs._OnViewRangeChanged;
+	TimeSliderArgs.OnClampRangeChanged = InArgs._OnClampRangeChanged;
 	TimeSliderArgs.ScrubPosition = InArgs._ScrubPosition;
 	TimeSliderArgs.OnBeginScrubberMovement = InArgs._OnBeginScrubbing;
 	TimeSliderArgs.OnEndScrubberMovement = InArgs._OnEndScrubbing;
@@ -205,9 +207,11 @@ void SSequencer::Construct( const FArguments& InArgs, TSharedRef< class FSequenc
 	bool bMirrorLabels = false;
 	// Create the top and bottom sliders
 	TSharedRef<ITimeSlider> TopTimeSlider = SequencerWidgets.CreateTimeSlider( TimeSliderController, bMirrorLabels );
-
 	bMirrorLabels = true;
-	TSharedRef<ITimeSlider> BottomTimeSlider = SequencerWidgets.CreateTimeSlider( TimeSliderController, bMirrorLabels);
+	TSharedRef<ITimeSlider> BottomTimeSlider = SequencerWidgets.CreateTimeSlider( TimeSliderController, TAttribute<EVisibility>(this, &SSequencer::GetBottomTimeSliderVisibility), bMirrorLabels );
+
+	// Create bottom time range slider
+	TSharedRef<ITimeSlider> BottomTimeRange = SequencerWidgets.CreateTimeRange( TimeSliderController, TAttribute<EVisibility>(this, &SSequencer::GetTimeRangeVisibility), TAttribute<bool>(this, &SSequencer::ShowFrameNumbers));
 
 	OnGetAddMenuContent = InArgs._OnGetAddMenuContent;
 
@@ -387,7 +391,17 @@ void SSequencer::Construct( const FArguments& InArgs, TSharedRef< class FSequenc
 							+ SVerticalBox::Slot()
 							.AutoHeight()
 							[
-								BottomTimeSlider
+								SNew( SOverlay )
+
+								+ SOverlay::Slot()
+								[
+									BottomTimeSlider
+								]
+
+								+ SOverlay::Slot()
+								[
+									BottomTimeRange
+								]
 							]
 
 							+ SVerticalBox::Slot()
@@ -586,9 +600,11 @@ TSharedRef<SWidget> SSequencer::MakeSnapMenu()
 {
 	FMenuBuilder MenuBuilder( false, Sequencer.Pin()->GetCommandBindings() );
 
-	MenuBuilder.BeginSection("Frames", LOCTEXT("SnappingMenuFramesHeader", "Frames") );
+	MenuBuilder.BeginSection("FramesRanges", LOCTEXT("SnappingMenuFrameRangesHeader", "Frame Ranges") );
 	{
 		MenuBuilder.AddMenuEntry( FSequencerCommands::Get().ToggleShowFrameNumbers );
+		MenuBuilder.AddMenuEntry( FSequencerCommands::Get().ToggleShowRangeSlider );
+		MenuBuilder.AddMenuEntry( FSequencerCommands::Get().ToggleLockInOutToStartEndRange );
 	}
 	MenuBuilder.EndSection();
 
@@ -1283,6 +1299,21 @@ EVisibility SSequencer::GetCurveEditorToolBarVisibility() const
 	return Settings->GetShowCurveEditor() ? EVisibility::Visible : EVisibility::Collapsed;
 }
 
+EVisibility SSequencer::GetBottomTimeSliderVisibility() const
+{
+	return Settings->GetShowRangeSlider() ? EVisibility::Hidden : EVisibility::Visible;
+}
+
+EVisibility SSequencer::GetTimeRangeVisibility() const
+{
+	return Settings->GetShowRangeSlider() ? EVisibility::Visible : EVisibility::Hidden;
+}
+
+bool SSequencer::ShowFrameNumbers() const
+{
+	return Sequencer.Pin()->CanShowFrameNumbers() && Settings->GetShowFrameNumbers();
+}
+
 float SSequencer::GetOutlinerSpacerFill() const
 {
 	const float Column1Coeff = GetColumnFillCoefficient(1);
diff --git a/Engine/Source/Editor/Sequencer/Private/SSequencer.h b/Engine/Source/Editor/Sequencer/Private/SSequencer.h
--- a/Engine/Source/Editor/Sequencer/Private/SSequencer.h
+++ b/Engine/Source/Editor/Sequencer/Private/SSequencer.h
@@ -54,10 +54,14 @@ public:
 	{}
 		/** The current view range (seconds) */
 		SLATE_ATTRIBUTE( FAnimatedRange, ViewRange )
+		/** The current clamp range (seconds) */
+		SLATE_ATTRIBUTE( FAnimatedRange, ClampRange )
 		/** The current scrub position in (seconds) */
 		SLATE_ATTRIBUTE( float, ScrubPosition )
 		/** Called when the user changes the view range */
 		SLATE_EVENT( FOnViewRangeChanged, OnViewRangeChanged )
+		/** Called when the user changes the clamp range */
+		SLATE_EVENT( FOnClampRangeChanged, OnClampRangeChanged )
 		/** Called when the user has begun scrubbing */
 		SLATE_EVENT( FSimpleDelegate, OnBeginScrubbing )
 		/** Called when the user has finished scrubbing */
@@ -227,9 +231,18 @@ private:
 	/** Gets whether or not the breadcrumb trail should be visible. */
 	EVisibility GetBreadcrumbTrailVisibility() const;
 
-	/** Gets whether or not the curve editor toolbar should be visibe. */
+	/** Gets whether or not the curve editor toolbar should be visible. */
 	EVisibility GetCurveEditorToolBarVisibility() const;
 
+	/** Gets whether or not the bottom time slider should be visible. */
+	EVisibility GetBottomTimeSliderVisibility() const;
+
+	/** Gets whether or not the time range should be visible. */
+	EVisibility GetTimeRangeVisibility() const;
+
+	/** Gets whether or not to show frame numbers. */
+	bool ShowFrameNumbers() const;
+
 	/** Called when a column fill percentage is changed by a splitter slot. */
 	void OnColumnFillCoefficientChanged(float FillCoefficient, int32 ColumnIndex);
 
diff --git a/Engine/Source/Editor/Sequencer/Private/Sequencer.cpp b/Engine/Source/Editor/Sequencer/Private/Sequencer.cpp
--- a/Engine/Source/Editor/Sequencer/Private/Sequencer.cpp
+++ b/Engine/Source/Editor/Sequencer/Private/Sequencer.cpp
@@ -100,11 +100,13 @@ void FSequencer::InitSequencer(const FSequencerInitParams& InitParams, const TSh
 		// Make internal widgets
 		SequencerWidget = SNew( SSequencer, SharedThis( this ) )
 			.ViewRange( this, &FSequencer::GetViewRange )
+			.ClampRange( this, &FSequencer::GetClampRange )
 			.ScrubPosition( this, &FSequencer::OnGetScrubPosition )
 			.OnBeginScrubbing( this, &FSequencer::OnBeginScrubbing )
 			.OnEndScrubbing( this, &FSequencer::OnEndScrubbing )
 			.OnScrubPositionChanged( this, &FSequencer::OnScrubPositionChanged )
 			.OnViewRangeChanged( this, &FSequencer::OnViewRangeChanged )
+			.OnClampRangeChanged( this, &FSequencer::OnClampRangeChanged )
 			.OnGetAddMenuContent(InitParams.ViewParams.OnGetAddMenuContent);
 
 		// When undo occurs, get a notification so we can make sure our view is up to date
@@ -145,6 +147,8 @@ void FSequencer::InitSequencer(const FSequencerInitParams& InitParams, const TSh
 		// Update initial movie scene data
 		NotifyMovieSceneDataChanged();
 
+		UpdateTimeBoundsToFocusedMovieScene();
+
 		// NOTE: Could fill in asset editor commands here!
 
 		BindSequencerCommands();
@@ -331,6 +335,8 @@ void FSequencer::ResetToNewAnimation(UMovieSceneAnimation& NewAnimation)
 	SequencerWidget->ResetBreadcrumbs();
 
 	NotifyMovieSceneDataChanged();
+
+	UpdateTimeBoundsToFocusedMovieScene();
 }
 
 TSharedRef<FMovieSceneInstance> FSequencer::GetRootMovieSceneInstance() const
@@ -359,6 +365,8 @@ void FSequencer::FocusSubMovieScene( TSharedRef<FMovieSceneInstance> SubMovieSce
 
 	// Update internal data for the new movie scene
 	NotifyMovieSceneDataChanged();
+
+	UpdateTimeBoundsToFocusedMovieScene();
 }
 
 TSharedRef<FMovieSceneInstance> FSequencer::GetInstanceForSubMovieSceneSection( UMovieSceneSection& SubMovieSceneSection ) const
@@ -381,6 +389,8 @@ void FSequencer::PopToMovieScene( TSharedRef<FMovieSceneInstance> SubMovieSceneI
 		ResetPerMovieSceneData();
 
 		NotifyMovieSceneDataChanged();
+
+		UpdateTimeBoundsToFocusedMovieScene();
 	}
 }
 
@@ -535,6 +545,11 @@ FAnimatedRange FSequencer::GetViewRange() const
 	return AnimatedRange;
 }
 
+FAnimatedRange FSequencer::GetClampRange() const
+{
+	return TRange<float>(GetFocusedMovieScene()->StartTime, GetFocusedMovieScene()->EndTime);
+}
+
 bool FSequencer::GetAutoKeyEnabled() const 
 {
 	return Settings->GetAutoKeyEnabled();
@@ -580,6 +595,11 @@ void FSequencer::SetGlobalTime( float NewTime )
 		}
 	}
 
+	if (Settings->GetLockInOutToStartEndRange())
+	{
+		NewTime = FMath::Clamp(NewTime, GetFocusedMovieScene()->StartTime, GetFocusedMovieScene()->EndTime);
+	}
+
 	// Update the position
 	ScrubPosition = NewTime;
 
@@ -913,6 +933,49 @@ EPlaybackMode::Type FSequencer::GetPlaybackMode() const
 		EPlaybackMode::Stopped;
 }
 
+void FSequencer::UpdateTimeBoundsToFocusedMovieScene()
+{
+	USequencerProjectSettings* ProjectSettings = GetMutableDefault<USequencerProjectSettings>();
+
+	const int32 InFrame = ProjectSettings->InFrame;
+	const int32 OutFrame = ProjectSettings->OutFrame;
+
+	float InTime = GetFocusedMovieScene()->InTime;
+	if (InTime >= FLT_MAX)
+	{
+		InTime = SequencerHelpers::FrameToTime(InFrame, 1.0f/Settings->GetTimeSnapInterval());
+	}
+
+	float OutTime = GetFocusedMovieScene()->OutTime;
+	if (OutTime <= -FLT_MAX)
+	{
+		OutTime = SequencerHelpers::FrameToTime(OutFrame, 1.0f/Settings->GetTimeSnapInterval());
+	}
+
+	float StartTime = GetFocusedMovieScene()->StartTime;
+	if (StartTime >= FLT_MAX)
+	{
+		StartTime = SequencerHelpers::FrameToTime(InFrame, 1.0f/Settings->GetTimeSnapInterval());
+	}
+
+	float EndTime = GetFocusedMovieScene()->EndTime;
+	if (EndTime <= -FLT_MAX)
+	{
+		EndTime = SequencerHelpers::FrameToTime(OutFrame, 1.0f/Settings->GetTimeSnapInterval());
+	}
+
+	// Set the clamp range to the movie scene's clamp range.
+	OnClampRangeChanged(TRange<float>(StartTime, EndTime));
+
+	// Set the view range to the movie scene's view range.
+	OnViewRangeChanged(TRange<float>(InTime, OutTime), EViewRangeInterpolation::Immediate);
+
+	// Make sure the current time is within the bounds
+	if (!TargetViewRange.Contains(ScrubPosition))
+	{
+		ScrubPosition = LastViewRange.GetLowerBoundValue();
+	}
+}
 
 TRange<float> FSequencer::GetTimeBounds() const
 {
@@ -961,6 +1024,42 @@ void FSequencer::OnViewRangeChanged( TRange<float> NewViewRange, EViewRangeInter
 
 	if (!NewViewRange.IsEmpty())
 	{
+		// If locked, clamp the new range to clamp range
+		if (Settings->GetLockInOutToStartEndRange())
+		{
+			if ( NewViewRange.GetLowerBoundValue() < GetFocusedMovieScene()->StartTime)
+			{
+				NewViewRange = TRange<float>(GetFocusedMovieScene()->StartTime, GetFocusedMovieScene()->StartTime + (NewViewRange.GetUpperBoundValue() - NewViewRange.GetLowerBoundValue()));
+			}
+			else if (NewViewRange.GetUpperBoundValue() > GetFocusedMovieScene()->EndTime)
+			{
+				NewViewRange = TRange<float>(GetFocusedMovieScene()->EndTime - (NewViewRange.GetUpperBoundValue() - NewViewRange.GetLowerBoundValue()), GetFocusedMovieScene()->EndTime);
+			}
+		}
+		// Otherwise, grow the clamp range to the new range
+		else
+		{
+			bool bNeedsClampSet = false;
+			float NewClampRangeMin = GetFocusedMovieScene()->StartTime;
+			if ( NewViewRange.GetLowerBoundValue() < NewClampRangeMin)
+			{
+				bNeedsClampSet = true;
+				NewClampRangeMin = NewViewRange.GetLowerBoundValue();
+			}
+	
+			float NewClampRangeMax = GetFocusedMovieScene()->EndTime;
+			if ( NewViewRange.GetUpperBoundValue() > NewClampRangeMax)
+			{
+				bNeedsClampSet = true;
+				NewClampRangeMax = NewViewRange.GetUpperBoundValue();
+			}
+
+			if (bNeedsClampSet)
+			{
+				OnClampRangeChanged(TRange<float>(NewClampRangeMin, NewClampRangeMax));
+			}
+		}
+
 		if (AnimationLengthSeconds != 0.f)
 		{
 			if (ZoomAnimation.GetCurve(0).DurationSeconds != AnimationLengthSeconds)
@@ -981,6 +1080,20 @@ void FSequencer::OnViewRangeChanged( TRange<float> NewViewRange, EViewRangeInter
 			TargetViewRange = LastViewRange = NewViewRange;
 			ZoomAnimation.JumpToEnd();
 		}
+
+		// Synchronize the in and out time of the focused movie scene
+		GetFocusedMovieScene()->InTime = NewViewRange.GetLowerBoundValue();
+		GetFocusedMovieScene()->OutTime = NewViewRange.GetUpperBoundValue();
+	}
+}
+
+void FSequencer::OnClampRangeChanged( TRange<float> NewClampRange )
+{
+	if (!NewClampRange.IsEmpty())
+	{
+		// Synchronize the in and out time of the focused movie scene
+		GetFocusedMovieScene()->StartTime = NewClampRange.GetLowerBoundValue();
+		GetFocusedMovieScene()->EndTime = NewClampRange.GetUpperBoundValue();
 	}
 }
 
@@ -1749,6 +1862,18 @@ void FSequencer::BindSequencerCommands()
 		FCanExecuteAction::CreateSP(this, &FSequencer::CanShowFrameNumbers ),
 		FIsActionChecked::CreateLambda( [this]{ return Settings->GetShowFrameNumbers(); } ) );
 
+	SequencerCommandBindings->MapAction(
+		Commands.ToggleShowRangeSlider,
+		FExecuteAction::CreateLambda( [this]{ Settings->SetShowRangeSlider( !Settings->GetShowRangeSlider() ); } ),
+		FCanExecuteAction::CreateSP(this, &FSequencer::CanShowFrameNumbers ),
+		FIsActionChecked::CreateLambda( [this]{ return Settings->GetShowRangeSlider(); } ) );
+
+	SequencerCommandBindings->MapAction(
+		Commands.ToggleLockInOutToStartEndRange,
+		FExecuteAction::CreateLambda( [this]{ Settings->SetLockInOutToStartEndRange( !Settings->GetLockInOutToStartEndRange() ); } ),
+		FCanExecuteAction::CreateSP(this, &FSequencer::CanShowFrameNumbers ),
+		FIsActionChecked::CreateLambda( [this]{ return Settings->GetLockInOutToStartEndRange(); } ) );
+
 	SequencerCommandBindings->MapAction(
 		Commands.ToggleIsSnapEnabled,
 		FExecuteAction::CreateLambda( [this]{ Settings->SetIsSnapEnabled( !Settings->GetIsSnapEnabled() ); } ),
diff --git a/Engine/Source/Editor/Sequencer/Private/Sequencer.h b/Engine/Source/Editor/Sequencer/Private/Sequencer.h
--- a/Engine/Source/Editor/Sequencer/Private/Sequencer.h
+++ b/Engine/Source/Editor/Sequencer/Private/Sequencer.h
@@ -93,6 +93,9 @@ public:
 	/** @return The current view range */
 	virtual FAnimatedRange GetViewRange() const override;
 
+	/** @return The current clamp range */
+	FAnimatedRange GetClampRange() const;
+
 public:
 
 	/** Access the user-supplied settings object */
@@ -270,6 +273,9 @@ protected:
 	/** Sets the actor CDO such that it is placed in front of the active perspective viewport camera, if we have one */
 	static void PlaceActorInFrontOfCamera( AActor* ActorCDO );
 
+	/** Update the time bounds to the focused movie scene */
+	void UpdateTimeBoundsToFocusedMovieScene();
+
 	/**
 	 * Gets the far time boundaries of the currently edited movie scene
 	 * If the scene has shots, it only takes the shot section boundaries
@@ -293,7 +299,12 @@ protected:
 	 */
 	void OnViewRangeChanged( TRange<float> NewViewRange, EViewRangeInterpolation Interpolation = EViewRangeInterpolation::Animated );
 
-
+	/**
+	 * Called when the clamp range is changed by the user
+	 *
+	 * @param	NewClampRange The new clapm range
+	 */
+	void OnClampRangeChanged( TRange<float> NewClampRange );
 
 	/**
 	 * Called when the scrub position is changed by the user
diff --git a/Engine/Source/Editor/Sequencer/Private/SequencerCommands.cpp b/Engine/Source/Editor/Sequencer/Private/SequencerCommands.cpp
--- a/Engine/Source/Editor/Sequencer/Private/SequencerCommands.cpp
+++ b/Engine/Source/Editor/Sequencer/Private/SequencerCommands.cpp
@@ -27,6 +27,8 @@ void FSequencerCommands::RegisterCommands()
 	UI_COMMAND( ToggleAutoScroll, "Auto Scroll", "Toggle auto-scroll: When enabled, automatically scrolls the sequencer view to keep the current time visible.", EUserInterfaceActionType::ToggleButton, FInputChord(EModifierKey::Shift, EKeys::S) );
 
 	UI_COMMAND( ToggleShowFrameNumbers, "Show frame numbers", "Enables and disables showing frame numbers.", EUserInterfaceActionType::ToggleButton, FInputChord() );
+	UI_COMMAND( ToggleShowRangeSlider, "Show range slider", "Enables and disables showing the time range slider.", EUserInterfaceActionType::ToggleButton, FInputChord() );
+	UI_COMMAND( ToggleLockInOutToStartEndRange, "Lock In/Out to Start/End", "Enables and disables locking the in/out to the start/end range while scrolling the view.", EUserInterfaceActionType::ToggleButton, FInputChord() );
 	UI_COMMAND( ToggleIsSnapEnabled, "Enable Snapping", "Enables and disables snapping.", EUserInterfaceActionType::ToggleButton, FInputChord() );
 
 	UI_COMMAND( ToggleSnapKeyTimesToInterval, "Snap to the interval", "Snap keys to the interval.", EUserInterfaceActionType::ToggleButton, FInputChord() );
diff --git a/Engine/Source/Editor/Sequencer/Private/SequencerCommands.h b/Engine/Source/Editor/Sequencer/Private/SequencerCommands.h
--- a/Engine/Source/Editor/Sequencer/Private/SequencerCommands.h
+++ b/Engine/Source/Editor/Sequencer/Private/SequencerCommands.h
@@ -57,6 +57,12 @@ public:
 	/** Turns show frame numbers on and off. */
 	TSharedPtr< FUICommandInfo > ToggleShowFrameNumbers;
 
+	/** Turns the range slider on and off. */
+	TSharedPtr< FUICommandInfo > ToggleShowRangeSlider;
+
+	/** Toggles whether to lock the in/out to the start/end range on scrolling. */
+	TSharedPtr< FUICommandInfo > ToggleLockInOutToStartEndRange;
+
 	/** Turns snapping on and off. */
 	TSharedPtr< FUICommandInfo > ToggleIsSnapEnabled;
 
diff --git a/Engine/Source/Editor/Sequencer/Private/SequencerCommonHelpers.cpp b/Engine/Source/Editor/Sequencer/Private/SequencerCommonHelpers.cpp
--- a/Engine/Source/Editor/Sequencer/Private/SequencerCommonHelpers.cpp
+++ b/Engine/Source/Editor/Sequencer/Private/SequencerCommonHelpers.cpp
@@ -53,3 +53,14 @@ void SequencerHelpers::GetDescendantNodes(TSharedRef<FSequencerDisplayNode> Disp
 	}
 }
 
+int32 SequencerHelpers::TimeToFrame(float Time, float FrameRate)
+{
+	int32 Frame = Time * FrameRate;
+	return FMath::RoundToInt(Frame);
+}
+
+float SequencerHelpers::FrameToTime(int32 Frame, float FrameRate)
+{
+	return Frame / FrameRate;
+}
+
diff --git a/Engine/Source/Editor/Sequencer/Private/SequencerCommonHelpers.h b/Engine/Source/Editor/Sequencer/Private/SequencerCommonHelpers.h
--- a/Engine/Source/Editor/Sequencer/Private/SequencerCommonHelpers.h
+++ b/Engine/Source/Editor/Sequencer/Private/SequencerCommonHelpers.h
@@ -18,6 +18,16 @@ public:
 	* Get descendant nodes
 	*/
 	static void GetDescendantNodes(TSharedRef<FSequencerDisplayNode> DisplayNode, TSet<TSharedRef<FSequencerDisplayNode> >& Nodes);
+
+	/**
+	* Convert time to frame
+	*/
+	static int32 TimeToFrame(float Time, float FrameRate);
+
+	/**
+	* Convert frame to time
+	*/
+	static float FrameToTime(int32 Frame, float FrameRate);
 };
 
 class SequencerSnapValues
diff --git a/Engine/Source/Editor/Sequencer/Private/SequencerSettings.cpp b/Engine/Source/Editor/Sequencer/Private/SequencerSettings.cpp
--- a/Engine/Source/Editor/Sequencer/Private/SequencerSettings.cpp
+++ b/Engine/Source/Editor/Sequencer/Private/SequencerSettings.cpp
@@ -7,6 +7,8 @@ USequencerSettings::USequencerSettings( const FObjectInitializer& ObjectInitiali
 {
 	bAutoKeyEnabled = false;
 	bShowFrameNumbers = true;
+	bShowRangeSlider = false;
+	bLockInOutToStartEndRange = false;
 	bIsSnapEnabled = true;
 	TimeSnapInterval = .05f;
 	bSnapKeyTimesToInterval = true;
@@ -51,6 +53,34 @@ void USequencerSettings::SetShowFrameNumbers(bool InbShowFrameNumbers)
 	}
 }
 
+bool USequencerSettings::GetShowRangeSlider() const
+{
+	return bShowRangeSlider;
+}
+
+void USequencerSettings::SetShowRangeSlider(bool InbShowRangeSlider)
+{
+	if ( bShowRangeSlider != InbShowRangeSlider )
+	{
+		bShowRangeSlider = InbShowRangeSlider;
+		SaveConfig();
+	}
+}
+
+bool USequencerSettings::GetLockInOutToStartEndRange() const
+{
+	return bLockInOutToStartEndRange;
+}
+
+void USequencerSettings::SetLockInOutToStartEndRange(bool InbLockInOutToStartEndRange)
+{
+	if ( bLockInOutToStartEndRange != InbLockInOutToStartEndRange )
+	{
+		bLockInOutToStartEndRange = InbLockInOutToStartEndRange;
+		SaveConfig();
+	}
+}
+
 bool USequencerSettings::GetIsSnapEnabled() const
 {
 	return bIsSnapEnabled;
@@ -266,4 +296,6 @@ ULevelEditorSequencerSettings::ULevelEditorSequencerSettings( const FObjectIniti
 	: Super( ObjectInitializer )
 {
 	TimeSnapInterval = 0.041667f;
+	bShowRangeSlider = true;
+	bLockInOutToStartEndRange = true;
 }
\ No newline at end of file
diff --git a/Engine/Source/Editor/Sequencer/Private/SequencerSettings.h b/Engine/Source/Editor/Sequencer/Private/SequencerSettings.h
--- a/Engine/Source/Editor/Sequencer/Private/SequencerSettings.h
+++ b/Engine/Source/Editor/Sequencer/Private/SequencerSettings.h
@@ -54,6 +54,16 @@ public:
 	/** Sets whether or not to show frame numbers. */
 	void SetShowFrameNumbers(bool InbShowFrameNumbers);
 
+	/** Gets whether or not to show the time range slider. */
+	bool GetShowRangeSlider() const;
+	/** Sets whether or not to show frame numbers. */
+	void SetShowRangeSlider(bool InbShowRangeSlider);
+
+	/** Gets whether or not the in and out are bounded by start and end when scrolling. */
+	bool GetLockInOutToStartEndRange() const;
+	/** Sets whether or not the in and out are bounded by start and end when scrolling. */
+	void SetLockInOutToStartEndRange(bool InbLockInOutToStartEndRange);
+
 	/** Gets whether or not snapping is enabled. */
 	bool GetIsSnapEnabled() const;
 	/** Sets whether or not snapping is enabled. */
@@ -137,6 +147,12 @@ protected:
 	UPROPERTY( config )
 	bool bShowFrameNumbers;
 
+	UPROPERTY( config )
+	bool bShowRangeSlider;
+
+	UPROPERTY( config )
+	bool bLockInOutToStartEndRange;
+
 	UPROPERTY( config )
 	bool bIsSnapEnabled;
 
diff --git a/Engine/Source/Editor/Sequencer/Private/TimeSliderController.cpp b/Engine/Source/Editor/Sequencer/Private/TimeSliderController.cpp
--- a/Engine/Source/Editor/Sequencer/Private/TimeSliderController.cpp
+++ b/Engine/Source/Editor/Sequencer/Private/TimeSliderController.cpp
@@ -17,14 +17,6 @@ namespace ScrubConstants
 	const float ScrollPanFraction = 0.1f;
 }
 
-namespace TimeConversion
-{
-	int32 TimeToFrame(float Time, float FrameRate)
-	{
-		return FMath::RoundToInt(Time * FrameRate);
-	}
-}
-
 /** Utility struct for converting between scrub range space and local/absolute screen space */
 struct FScrubRangeToScreen
 {
@@ -162,8 +154,6 @@ void FSequencerTimeSliderController::DrawTicks( FSlateWindowElementList& OutDraw
 	TArray<FVector2D> LinePoints;
 	LinePoints.AddUninitialized(2);
 
-	float SnapSize = TimeSliderArgs.Settings->GetTimeSnapInterval();
-
 	float Seconds = 0;
 	while( (Seconds = OffsetNum*Spacing) < RangeToScreen.ViewInput.GetUpperBoundValue() )
 	{
@@ -199,7 +189,7 @@ void FSequencerTimeSliderController::DrawTicks( FSlateWindowElementList& OutDraw
 				FString FrameString;
 				if (SequencerSnapValues::IsTimeSnapIntervalFrameRate(TimeSliderArgs.Settings->GetTimeSnapInterval()) && TimeSliderArgs.Settings->GetShowFrameNumbers())
 				{
-					FrameString = FString::Printf( TEXT("%d"), TimeConversion::TimeToFrame(Seconds, 1.0f/SnapSize));
+					FrameString = FString::Printf( TEXT("%d"), TimeToFrame(Seconds));
 				}
 				else
 				{
@@ -313,8 +303,7 @@ int32 FSequencerTimeSliderController::OnPaintTimeSlider( bool bMirrorLabels, con
 		FString FrameString;
 		if (SequencerSnapValues::IsTimeSnapIntervalFrameRate(TimeSliderArgs.Settings->GetTimeSnapInterval()) && TimeSliderArgs.Settings->GetShowFrameNumbers())
 		{
-			float SnapSize = TimeSliderArgs.Settings->GetTimeSnapInterval();
-			FrameString = FString::Printf( TEXT("%d"), TimeConversion::TimeToFrame(Time, 1.0f/SnapSize));
+			FrameString = FString::Printf( TEXT("%d"), TimeToFrame(Time));
 		}
 		else
 		{
@@ -513,6 +502,8 @@ FReply FSequencerTimeSliderController::OnMouseMove( TSharedRef<SWidget> WidgetOw
 				float NewViewOutputMin = LocalViewRangeMin - InputDelta.X;
 				float NewViewOutputMax = LocalViewRangeMax - InputDelta.X;
 
+				const bool bMaintainRange = true;
+				ClampViewRange(NewViewOutputMin, NewViewOutputMax, bMaintainRange);
 				SetViewRange(NewViewOutputMin, NewViewOutputMax, EViewRangeInterpolation::Immediate);
 			}
 		}
@@ -651,22 +642,70 @@ int32 FSequencerTimeSliderController::OnPaintSectionView( const FGeometry& Allot
 	return LayerId;
 }
 
-void FSequencerTimeSliderController::SetViewRange( float NewRangeMin, float NewRangeMax, EViewRangeInterpolation Interpolation )
+int32 FSequencerTimeSliderController::TimeToFrame(float Time) const
 {
-	TOptional<float> LocalClampMin = TimeSliderArgs.ClampMin.Get();
-	TOptional<float> LocalClampMax = TimeSliderArgs.ClampMax.Get();
+	float FrameRate = 1.0f/TimeSliderArgs.Settings->GetTimeSnapInterval();
+	return SequencerHelpers::TimeToFrame(Time, FrameRate);
+}
+
+float FSequencerTimeSliderController::FrameToTime(int32 Frame) const
+{
+	float FrameRate = 1.0f/TimeSliderArgs.Settings->GetTimeSnapInterval();
+	return SequencerHelpers::FrameToTime(Frame, FrameRate);
+}
 
-	// Clamp the range if clamp values are set
-	if ( LocalClampMin.IsSet() && NewRangeMin < LocalClampMin.GetValue() )
+void FSequencerTimeSliderController::ClampViewRange(float& NewRangeMin, float& NewRangeMax, bool bMaintainRange)
+{
+	// If locked, clamp the new range to clamp range
+	if (TimeSliderArgs.Settings->GetLockInOutToStartEndRange())
 	{
-		NewRangeMin = LocalClampMin.GetValue();
-	}
+		if (bMaintainRange)
+		{
+			if ( NewRangeMin < TimeSliderArgs.ClampRange.Get().GetLowerBoundValue() )
+			{
+				NewRangeMin = TimeSliderArgs.ClampRange.Get().GetLowerBoundValue();
+				NewRangeMax = NewRangeMin + (TimeSliderArgs.ViewRange.Get().GetUpperBoundValue() - TimeSliderArgs.ViewRange.Get().GetLowerBoundValue());
+			}
 	
-	if ( LocalClampMax.IsSet() && NewRangeMax > LocalClampMax.GetValue() )
+			else if ( NewRangeMax > TimeSliderArgs.ClampRange.Get().GetUpperBoundValue() )
+			{
+				NewRangeMax = TimeSliderArgs.ClampRange.Get().GetUpperBoundValue();
+				NewRangeMin = NewRangeMax - (TimeSliderArgs.ViewRange.Get().GetUpperBoundValue() - TimeSliderArgs.ViewRange.Get().GetLowerBoundValue());
+			}
+		}
+		else
+		{
+			NewRangeMin = FMath::Clamp(NewRangeMin, TimeSliderArgs.ClampRange.Get().GetLowerBoundValue(), TimeSliderArgs.ClampRange.Get().GetUpperBoundValue());
+			NewRangeMax = FMath::Clamp(NewRangeMax, TimeSliderArgs.ClampRange.Get().GetLowerBoundValue(), TimeSliderArgs.ClampRange.Get().GetUpperBoundValue());
+		}
+	}
+	// Otherwise, grow the clamp range to the new range
+	else
 	{
-		NewRangeMax = LocalClampMax.GetValue();
+		bool bNeedsClampSet = false;
+		float NewClampRangeMin = TimeSliderArgs.ClampRange.Get().GetLowerBoundValue();
+		if ( NewRangeMin < TimeSliderArgs.ClampRange.Get().GetLowerBoundValue() )
+		{
+			NewClampRangeMin = NewRangeMin;
+			bNeedsClampSet = true;
+		}
+	
+		float NewClampRangeMax = TimeSliderArgs.ClampRange.Get().GetUpperBoundValue();
+		if ( NewRangeMax > TimeSliderArgs.ClampRange.Get().GetUpperBoundValue() )
+		{
+			NewClampRangeMax = NewRangeMax;
+			bNeedsClampSet = true;
+		}
+
+		if (bNeedsClampSet)
+		{
+			SetClampRange(NewClampRangeMin, NewClampRangeMax);
+		}
 	}
+}
 
+void FSequencerTimeSliderController::SetViewRange( float NewRangeMin, float NewRangeMax, EViewRangeInterpolation Interpolation )
+{
 	const TRange<float> NewRange(NewRangeMin, NewRangeMax);
 
 	TimeSliderArgs.OnViewRangeChanged.ExecuteIfBound( NewRange, Interpolation );
@@ -678,6 +717,19 @@ void FSequencerTimeSliderController::SetViewRange( float NewRangeMin, float NewR
 	}
 }
 
+void FSequencerTimeSliderController::SetClampRange( float NewRangeMin, float NewRangeMax )
+{
+	const TRange<float> NewRange(NewRangeMin, NewRangeMax);
+
+	TimeSliderArgs.OnClampRangeChanged.ExecuteIfBound(NewRange);
+
+	if( !TimeSliderArgs.ClampRange.IsBound() )
+	{	
+		// The  output is not bound to a delegate so we'll manage the value ourselves (no animation)
+		TimeSliderArgs.ClampRange.Set(NewRange);
+	}
+}
+
 bool FSequencerTimeSliderController::ZoomByDelta( float InDelta, float MousePositionFraction )
 {
 	TRange<float> LocalViewRange = TimeSliderArgs.ViewRange.Get().GetAnimationTarget();
@@ -691,6 +743,8 @@ bool FSequencerTimeSliderController::ZoomByDelta( float InDelta, float MousePosi
 
 	if( FMath::Abs( OutputChange ) > 0.01f && NewViewOutputMin < NewViewOutputMax )
 	{
+		const bool bMaintainRange = false;
+		ClampViewRange(NewViewOutputMin, NewViewOutputMax, bMaintainRange);
 		SetViewRange(NewViewOutputMin, NewViewOutputMax, EViewRangeInterpolation::Animated);
 		return true;
 	}
@@ -708,7 +762,12 @@ void FSequencerTimeSliderController::PanByDelta( float InDelta )
 	// Adjust the delta to be a percentage of the current range
 	InDelta *= ScrubConstants::ScrollPanFraction * (CurrentMax - CurrentMin);
 
-	SetViewRange(CurrentMin + InDelta, CurrentMax + InDelta, EViewRangeInterpolation::Animated);
+	float NewViewOutputMin = CurrentMin + InDelta;
+	float NewViewOutputMax = CurrentMax + InDelta;
+
+	const bool bMaintainRange = true;
+	ClampViewRange(NewViewOutputMin, NewViewOutputMax, bMaintainRange);
+	SetViewRange(NewViewOutputMin, NewViewOutputMax, EViewRangeInterpolation::Animated);
 }
 
 
diff --git a/Engine/Source/Editor/Sequencer/Private/TimeSliderController.h b/Engine/Source/Editor/Sequencer/Private/TimeSliderController.h
--- a/Engine/Source/Editor/Sequencer/Private/TimeSliderController.h
+++ b/Engine/Source/Editor/Sequencer/Private/TimeSliderController.h
@@ -22,7 +22,26 @@ public:
 	/** End ITimeSliderController Interface */
 
 	/** Get the current view range for this controller */
-	FAnimatedRange GetViewRange() const { return TimeSliderArgs.ViewRange.Get(); }
+	virtual FAnimatedRange GetViewRange() const override { return TimeSliderArgs.ViewRange.Get(); }
+
+	/** Get the current clamp range for this controller */
+	//virtual FAnimatedRange GetClampRange() const override { return TRange<float>(TimeSliderArgs.ClampMin.Get().GetValue(), TimeSliderArgs.ClampMax.Get().GetValue()); }
+	virtual FAnimatedRange GetClampRange() const override { return TimeSliderArgs.ClampRange.Get(); }
+
+	/** Convert time to frame */
+	virtual int32 TimeToFrame(float Time) const override;
+
+	/** Convert time to frame */
+	virtual float FrameToTime(int32 Frame) const override;
+
+	/**
+	 * Clamp the given range to the clamp range 
+	 *
+	 * @param NewRangeMin		The new lower bound of the range
+	 * @param NewRangeMax		The new upper bound of the range
+	 * @param bMaintainRange	Maintain the current range size if 
+	 */	
+	void ClampViewRange(float& NewRangeMin, float& NewRangeMax, bool bMaintainRange);
 
 	/**
 	 * Set a new range based on a min, max and an interpolation mode
@@ -31,7 +50,15 @@ public:
 	 * @param NewRangeMax		The new upper bound of the range
 	 * @param Interpolation		How to set the new range (either immediately, or animated)
 	 */
-	void SetViewRange( float NewRangeMin, float NewRangeMax, EViewRangeInterpolation Interpolation );
+	virtual void SetViewRange( float NewRangeMin, float NewRangeMax, EViewRangeInterpolation Interpolation ) override;
+
+	/**
+	 * Set a new clamp range based on a min, max
+	 * 
+	 * @param NewRangeMin		The new lower bound of the clamp range
+	 * @param NewRangeMax		The new upper bound of the clamp range
+	 */
+	virtual void SetClampRange( float NewRangeMin, float NewRangeMax ) override;
 
 	/**
 	 * Zoom the range by a given delta.
diff --git a/Engine/Source/Editor/SequencerWidgets/Private/STimeRange.cpp b/Engine/Source/Editor/SequencerWidgets/Private/STimeRange.cpp
new file mode 100644
--- /dev/null
+++ b/Engine/Source/Editor/SequencerWidgets/Private/STimeRange.cpp
@@ -0,0 +1,337 @@
+// Copyright 1998-2015 Epic Games, Inc. All Rights Reserved.
+
+#include "SequencerWidgetsPrivatePCH.h"
+#include "STimeRange.h"
+#include "STimeRangeSlider.h"
+#include "SlateStyle.h"
+
+#define LOCTEXT_NAMESPACE "STimeRange"
+
+/*
+* Custom time range text box for the inputs
+*/
+class STimeRangeTextBox : public SEditableTextBox
+{
+public:
+	void Construct(const FArguments& InArgs) { SEditableTextBox::Construct(InArgs); }
+protected:
+	virtual FVector2D ComputeDesiredSize(float) const override { return FVector2D(46.f, 16.f); }
+};
+
+void STimeRange::Construct( const STimeRange::FArguments& InArgs, TSharedRef<ITimeSliderController> InTimeSliderController )
+{
+	TimeSliderController = InTimeSliderController;
+	ShowFrameNumbers = InArgs._ShowFrameNumbers;
+	
+	this->ChildSlot
+	.HAlign(HAlign_Fill)
+	[
+		SNew(SHorizontalBox)
+		+SHorizontalBox::Slot()
+		.AutoWidth()
+		.Padding(2.0f, 2.0f)
+		[
+			SNew(STimeRangeTextBox)
+			.Text(this, &STimeRange::StartTime)
+			.ToolTipText(this, &STimeRange::StartTimeTooltip )
+			.OnTextCommitted( this, &STimeRange::OnStartTimeCommitted )
+		]
+		+SHorizontalBox::Slot()
+		.AutoWidth()
+		.Padding(2.0f, 2.0f)
+		[
+			SNew(STimeRangeTextBox)			
+			.Text(this, &STimeRange::InTime)
+			.ToolTipText(this, &STimeRange::InTimeTooltip)
+			.OnTextCommitted( this, &STimeRange::OnInTimeCommitted )
+		]
+		+SHorizontalBox::Slot()
+		.HAlign(HAlign_Fill)
+		.VAlign(VAlign_Fill)
+		.Padding(2.0f, 4.0f)
+		[
+			SAssignNew(TimeRangeSlider, STimeRangeSlider, InTimeSliderController)
+		]
+		+SHorizontalBox::Slot()
+		.AutoWidth()
+		.Padding(2.0f, 2.0f)
+		[
+			SNew(STimeRangeTextBox)
+			.Text(this, &STimeRange::OutTime)
+			.ToolTipText(this, &STimeRange::OutTimeTooltip)
+			.OnTextCommitted( this, &STimeRange::OnOutTimeCommitted )
+		]
+		+SHorizontalBox::Slot()
+		.AutoWidth()
+		.Padding(2.0f, 2.0f)
+		[
+			SNew(STimeRangeTextBox)
+			.Text(this, &STimeRange::EndTime)
+			.ToolTipText(this, &STimeRange::EndTimeTooltip)
+			.OnTextCommitted( this, &STimeRange::OnEndTimeCommitted )
+		]
+	];
+}
+
+FText STimeRange::StartTime() const
+{
+	if (TimeSliderController.IsValid())
+	{
+		float StartTime = TimeSliderController.Get()->GetClampRange().GetLowerBoundValue();
+		bool bShowFrameNumbers = ShowFrameNumbers.IsBound() ? ShowFrameNumbers.Get() : false;
+
+		FString StartTimeStr;
+		if (bShowFrameNumbers)
+		{
+			int32 Frame = TimeSliderController.Get()->TimeToFrame(StartTime);
+			StartTimeStr = FString::Printf( TEXT("%d"), Frame);
+		}
+		else
+		{
+			StartTimeStr = FString::Printf( TEXT("%.2f"), StartTime);
+		}
+
+		return FText::FromString(StartTimeStr);
+	}
+
+	return FText::FromString("");
+}
+
+FText STimeRange::EndTime() const
+{
+	if (TimeSliderController.IsValid())
+	{
+		float EndTime = TimeSliderController.Get()->GetClampRange().GetUpperBoundValue();
+		bool bShowFrameNumbers = ShowFrameNumbers.IsBound() ? ShowFrameNumbers.Get() : false;
+
+		FString EndTimeStr;
+		if (bShowFrameNumbers)
+		{
+			int32 Frame = TimeSliderController.Get()->TimeToFrame(EndTime);
+			EndTimeStr = FString::Printf( TEXT("%d"), Frame);
+		}
+		else
+		{
+			EndTimeStr = FString::Printf( TEXT("%.2f"), EndTime);
+		}
+
+		return FText::FromString(EndTimeStr);
+	}
+
+	return FText::FromString("");
+}
+
+FText STimeRange::InTime() const
+{
+	if (TimeSliderController.IsValid())
+	{
+		float InTime = TimeSliderController.Get()->GetViewRange().GetLowerBoundValue();
+		bool bShowFrameNumbers = ShowFrameNumbers.IsBound() ? ShowFrameNumbers.Get() : false;
+
+		FString InTimeStr;
+		if (bShowFrameNumbers)
+		{
+			int32 Frame = TimeSliderController.Get()->TimeToFrame(InTime);
+			InTimeStr = FString::Printf( TEXT("%d"), Frame);
+		}
+		else
+		{
+			InTimeStr = FString::Printf( TEXT("%.2f"), InTime);
+		}
+
+		return FText::FromString(InTimeStr);
+	}
+
+	return FText::FromString("");
+}
+
+FText STimeRange::OutTime() const
+{
+	if (TimeSliderController.IsValid())
+	{
+		float OutTime = TimeSliderController.Get()->GetViewRange().GetUpperBoundValue();
+		bool bShowFrameNumbers = ShowFrameNumbers.IsBound() ? ShowFrameNumbers.Get() : false;
+
+		FString OutTimeStr;
+		if (bShowFrameNumbers)
+		{
+			int32 Frame = TimeSliderController.Get()->TimeToFrame(OutTime);
+			OutTimeStr = FString::Printf( TEXT("%d"), Frame);
+		}
+		else
+		{
+			OutTimeStr = FString::Printf( TEXT("%.2f"), OutTime); 
+		}
+
+		return FText::FromString(OutTimeStr);
+	}
+
+	return FText::FromString("");
+}
+
+void STimeRange::OnStartTimeCommitted(const FText& NewText, ETextCommit::Type InTextCommit)
+{
+	FString NewTextStr = NewText.ToString();
+
+	if (TimeSliderController.IsValid())
+	{
+		float NewStart;
+		bool bShowFrameNumbers = ShowFrameNumbers.IsBound() ? ShowFrameNumbers.Get() : false;
+
+		if (bShowFrameNumbers)
+		{
+			int32 NewStartFrame = FCString::Atoi(*NewTextStr);
+			NewStart = TimeSliderController.Get()->FrameToTime(NewStartFrame);
+		}
+		else
+		{
+			NewStart = FCString::Atof(*NewTextStr);
+		}
+
+		if (NewStart > TimeSliderController.Get()->GetViewRange().GetLowerBoundValue())
+		{
+			TimeSliderController.Get()->SetViewRange(NewStart, TimeSliderController.Get()->GetViewRange().GetUpperBoundValue(), EViewRangeInterpolation::Immediate);
+		}
+
+		TimeSliderController.Get()->SetClampRange(NewStart, TimeSliderController.Get()->GetClampRange().GetUpperBoundValue());
+	}
+}
+
+void STimeRange::OnEndTimeCommitted(const FText& NewText, ETextCommit::Type InTextCommit)
+{
+	FString NewTextStr = NewText.ToString();
+
+	if (TimeSliderController.IsValid())
+	{
+		float NewEnd;
+		bool bShowFrameNumbers = ShowFrameNumbers.IsBound() ? ShowFrameNumbers.Get() : false;
+
+		if (bShowFrameNumbers)
+		{
+			int32 NewEndFrame = FCString::Atoi(*NewTextStr);
+			NewEnd = TimeSliderController.Get()->FrameToTime(NewEndFrame);
+		}
+		else
+		{
+			NewEnd = FCString::Atof(*NewTextStr);
+		}
+
+		if (NewEnd < TimeSliderController.Get()->GetViewRange().GetUpperBoundValue())
+		{
+			TimeSliderController.Get()->SetViewRange(TimeSliderController.Get()->GetViewRange().GetLowerBoundValue(), NewEnd, EViewRangeInterpolation::Immediate);
+		}
+
+		TimeSliderController.Get()->SetClampRange(TimeSliderController.Get()->GetClampRange().GetLowerBoundValue(), NewEnd);
+	}
+}
+
+void STimeRange::OnInTimeCommitted(const FText& NewText, ETextCommit::Type InTextCommit)
+{
+	FString NewTextStr = NewText.ToString();
+
+	if (TimeSliderController.IsValid())
+	{
+		float NewIn;
+		bool bShowFrameNumbers = ShowFrameNumbers.IsBound() ? ShowFrameNumbers.Get() : false;
+
+		if (bShowFrameNumbers)
+		{
+			int32 NewInFrame = FCString::Atoi(*NewTextStr);
+
+			NewIn = TimeSliderController.Get()->FrameToTime(NewInFrame);
+		}
+		else
+		{
+			NewIn = FCString::Atof(*NewTextStr);
+		}
+
+		if (NewIn < TimeSliderController.Get()->GetClampRange().GetLowerBoundValue())
+		{
+			TimeSliderController.Get()->SetClampRange(NewIn, TimeSliderController.Get()->GetClampRange().GetUpperBoundValue());
+		}
+
+		TimeSliderController.Get()->SetViewRange(NewIn, TimeSliderController.Get()->GetViewRange().GetUpperBoundValue(), EViewRangeInterpolation::Immediate);
+	}
+}
+
+void STimeRange::OnOutTimeCommitted(const FText& NewText, ETextCommit::Type InTextCommit)
+{
+	FString NewTextStr = NewText.ToString();
+
+	if (TimeSliderController.IsValid())
+	{
+		float NewOut;
+		bool bShowFrameNumbers = ShowFrameNumbers.IsBound() ? ShowFrameNumbers.Get() : false;
+
+		if (bShowFrameNumbers)
+		{
+			int32 NewOutFrame = FCString::Atoi(*NewTextStr);
+			NewOut = TimeSliderController.Get()->FrameToTime(NewOutFrame);
+		}
+		else
+		{
+			NewOut = FCString::Atof(*NewTextStr);
+		}
+
+		if (NewOut > TimeSliderController.Get()->GetClampRange().GetUpperBoundValue())
+		{
+			TimeSliderController.Get()->SetClampRange(TimeSliderController.Get()->GetClampRange().GetLowerBoundValue(), NewOut);
+		}
+
+		TimeSliderController.Get()->SetViewRange(TimeSliderController.Get()->GetViewRange().GetLowerBoundValue(), NewOut, EViewRangeInterpolation::Immediate);
+	}
+}
+
+FText STimeRange::StartTimeTooltip() const
+{
+	bool bShowFrameNumbers = ShowFrameNumbers.IsBound() ? ShowFrameNumbers.Get() : false;
+	if (bShowFrameNumbers)
+	{
+		return LOCTEXT("StartFrameTooltip", "Start frame");
+	}
+	else
+	{
+		return LOCTEXT("StartTimeTooltip", "Start time");
+	}
+}
+
+FText STimeRange::InTimeTooltip() const
+{
+	bool bShowFrameNumbers = ShowFrameNumbers.IsBound() ? ShowFrameNumbers.Get() : false;
+	if (bShowFrameNumbers)
+	{
+		return LOCTEXT("InFrameTooltip", "In frame");
+	}
+	else
+	{
+		return LOCTEXT("InTimeTooltip", "In time");
+	}
+}
+
+FText STimeRange::OutTimeTooltip() const
+{
+	bool bShowFrameNumbers = ShowFrameNumbers.IsBound() ? ShowFrameNumbers.Get() : false;
+	if (bShowFrameNumbers)
+	{
+		return LOCTEXT("OutFrameTooltip", "Out frame");
+	}
+	else
+	{
+		return LOCTEXT("OutTimeTooltip", "Out time");
+	}
+}
+
+FText STimeRange::EndTimeTooltip() const
+{
+	bool bShowFrameNumbers = ShowFrameNumbers.IsBound() ? ShowFrameNumbers.Get() : false;
+	if (bShowFrameNumbers)
+	{
+		return LOCTEXT("EndFrameTooltip", "End frame");
+	}
+	else
+	{
+		return LOCTEXT("EndTimeTooltip", "End time");
+	}
+}
+
+#undef LOCTEXT_NAMESPACE
diff --git a/Engine/Source/Editor/SequencerWidgets/Private/STimeRange.h b/Engine/Source/Editor/SequencerWidgets/Private/STimeRange.h
new file mode 100644
--- /dev/null
+++ b/Engine/Source/Editor/SequencerWidgets/Private/STimeRange.h
@@ -0,0 +1,51 @@
+// Copyright 1998-2015 Epic Games, Inc. All Rights Reserved.
+
+#pragma once
+
+#include "ITimeSlider.h"
+
+class STimeRangeSlider;
+
+DECLARE_DELEGATE_RetVal( bool, STimeRangeGetter );
+
+class STimeRange : public ITimeSlider
+{
+public:
+
+	//SLATE_BEGIN_ARGS(STimeRange){}
+	//	SLATE_DEFAULT_SLOT(FArguments, Content)
+	//SLATE_END_ARGS()
+	SLATE_BEGIN_ARGS(STimeRange)
+	{}
+		/* If we should show frame numbers on the timeline */
+		SLATE_ARGUMENT( TAttribute<bool>, ShowFrameNumbers )
+	SLATE_END_ARGS()
+
+	/**
+	 * Construct the widget
+	 * 
+	 * @param InArgs   A declaration from which to construct the widget
+	 */
+	void Construct( const FArguments& InArgs, TSharedRef<ITimeSliderController> InTimeSliderController );
+
+protected:
+	FText InTime() const;
+	FText OutTime() const;
+	FText StartTime() const;
+	FText EndTime() const;
+
+	void OnStartTimeCommitted(const FText& NewText, ETextCommit::Type InTextCommit);
+	void OnEndTimeCommitted(const FText& NewText, ETextCommit::Type InTextCommit);
+	void OnInTimeCommitted(const FText& NewText, ETextCommit::Type InTextCommit);
+	void OnOutTimeCommitted(const FText& NewText, ETextCommit::Type InTextCommit);
+
+	FText InTimeTooltip() const;
+	FText OutTimeTooltip() const;
+	FText StartTimeTooltip() const;
+	FText EndTimeTooltip() const;
+
+private:
+	TSharedPtr<ITimeSliderController> TimeSliderController;
+	TSharedPtr<STimeRangeSlider> TimeRangeSlider;
+	TAttribute<bool> ShowFrameNumbers;
+};
\ No newline at end of file
diff --git a/Engine/Source/Editor/SequencerWidgets/Private/STimeRangeSlider.cpp b/Engine/Source/Editor/SequencerWidgets/Private/STimeRangeSlider.cpp
new file mode 100644
--- /dev/null
+++ b/Engine/Source/Editor/SequencerWidgets/Private/STimeRangeSlider.cpp
@@ -0,0 +1,284 @@
+// Copyright 1998-2015 Epic Games, Inc. All Rights Reserved.
+
+#include "SequencerWidgetsPrivatePCH.h"
+#include "STimeRangeSlider.h"
+#include "SlateStyle.h"
+
+#define LOCTEXT_NAMESPACE "STimeRangeSlider"
+
+namespace TimeRangeSliderConstants
+{
+	const int32 HandleSize = 14;
+}
+
+void STimeRangeSlider::Construct( const FArguments& InArgs, TSharedRef<ITimeSliderController> InTimeSliderController)
+{
+	TimeSliderController = InTimeSliderController;
+	LastViewRange.Empty();
+
+	ResetState();
+	ResetHoveredState();
+}
+
+float STimeRangeSlider::ComputeDragDelta(const FPointerEvent& MouseEvent, int32 GeometryWidth) const
+{
+	float StartTime = 0;
+	float EndTime = 0;
+
+	if (TimeSliderController.IsValid())
+	{
+		StartTime = TimeSliderController.Get()->GetClampRange().GetLowerBoundValue();
+		EndTime = TimeSliderController.Get()->GetClampRange().GetUpperBoundValue();
+	}
+	float DragDistance = (MouseEvent.GetScreenSpacePosition() - MouseDownPosition).X;
+
+	return DragDistance / (float)GeometryWidth * (EndTime - StartTime);
+}
+
+void STimeRangeSlider::ComputeHandleOffsets(float& LeftHandleOffset, float& HandleOffset, float& RightHandleOffset, int32 GeometryWidth) const
+{
+	float StartTime = 0;
+	float InTime = 0;
+	float OutTime = 0;
+	float EndTime = 0;
+
+	if (TimeSliderController.IsValid())
+	{
+		StartTime = TimeSliderController.Get()->GetClampRange().GetLowerBoundValue();
+		InTime = TimeSliderController.Get()->GetViewRange().GetLowerBoundValue();
+		OutTime = TimeSliderController.Get()->GetViewRange().GetUpperBoundValue();
+		EndTime = TimeSliderController.Get()->GetClampRange().GetUpperBoundValue();
+	}
+
+	LeftHandleOffset = (InTime - StartTime) / (EndTime - StartTime) * (float)GeometryWidth;
+	HandleOffset = LeftHandleOffset + TimeRangeSliderConstants::HandleSize;
+	RightHandleOffset = (OutTime - StartTime) / (EndTime - StartTime) * (float)GeometryWidth - TimeRangeSliderConstants::HandleSize;
+}
+
+int32 STimeRangeSlider::OnPaint( const FPaintArgs& Args, const FGeometry& AllottedGeometry, const FSlateRect& MyClippingRect, FSlateWindowElementList& OutDrawElements, int32 LayerId, const FWidgetStyle& InWidgetStyle, bool bParentEnabled ) const
+{
+	const int32 BackgroundLayer = LayerId+1;
+	const int32 SliderBoxLayer = BackgroundLayer+1;
+	const int32 HandleLayer = SliderBoxLayer+1;
+
+	static const FSlateBrush* RangeHandleLeft = FEditorStyle::GetBrush( TEXT( "Sequencer.Timeline.RangeHandleLeft" ) ); 
+	static const FSlateBrush* RangeHandleRight = FEditorStyle::GetBrush( TEXT( "Sequencer.Timeline.RangeHandleRight" ) ); 
+	static const FSlateBrush* RangeHandle = FEditorStyle::GetBrush( TEXT( "Sequencer.Timeline.RangeHandle" ) ); 
+
+	float LeftHandleOffset = 0.f;
+	float HandleOffset = 0.f;
+	float RightHandleOffset = 0.f;
+	ComputeHandleOffsets(LeftHandleOffset, HandleOffset, RightHandleOffset, AllottedGeometry.Size.X);
+
+	static const FName SelectionColorName("SelectionColor");
+	FLinearColor SelectionColor = FEditorStyle::GetSlateColor(SelectionColorName).GetColor(FWidgetStyle());
+
+	// Draw the left handle box
+	FSlateDrawElement::MakeBox( 
+		OutDrawElements,
+		LayerId, 
+		AllottedGeometry.ToPaintGeometry(FVector2D(LeftHandleOffset, 0.0f), FVector2D(TimeRangeSliderConstants::HandleSize, TimeRangeSliderConstants::HandleSize)),
+		RangeHandleLeft,
+		MyClippingRect,
+		ESlateDrawEffect::None,
+		(bLeftHandleDragged || bLeftHandleHovered) ? SelectionColor : FLinearColor::Gray);
+
+	// Draw the right handle box
+	FSlateDrawElement::MakeBox( 
+		OutDrawElements,
+		LayerId, 
+		AllottedGeometry.ToPaintGeometry(FVector2D(RightHandleOffset, 0.0f), FVector2D(TimeRangeSliderConstants::HandleSize, TimeRangeSliderConstants::HandleSize)),
+		RangeHandleRight,
+		MyClippingRect,
+		ESlateDrawEffect::None,
+		(bRightHandleDragged || bRightHandleHovered) ? SelectionColor : FLinearColor::Gray);
+
+	// Draw the handle box
+	FSlateDrawElement::MakeBox( 
+		OutDrawElements,
+		LayerId, 
+		AllottedGeometry.ToPaintGeometry(FVector2D(HandleOffset, 0.0f), FVector2D(RightHandleOffset-LeftHandleOffset-TimeRangeSliderConstants::HandleSize, TimeRangeSliderConstants::HandleSize)),
+		RangeHandle,
+		MyClippingRect,
+		ESlateDrawEffect::None,
+		(bHandleDragged || bHandleHovered) ? SelectionColor : FLinearColor::Gray);
+
+	SCompoundWidget::OnPaint(Args, AllottedGeometry, MyClippingRect, OutDrawElements, LayerId, InWidgetStyle, ShouldBeEnabled( bParentEnabled ));
+
+	return LayerId;
+}
+
+FReply STimeRangeSlider::OnMouseButtonDown( const FGeometry& MyGeometry, const FPointerEvent& MouseEvent )
+{
+	MouseDownPosition = MouseEvent.GetScreenSpacePosition();
+	if (TimeSliderController.IsValid())
+	{		
+		MouseDownViewRange = TimeSliderController.Get()->GetViewRange();
+	}
+
+	if (bHandleHovered)
+	{
+		bHandleDragged = true;
+		return FReply::Handled().CaptureMouse(AsShared());
+	}
+	else if (bLeftHandleHovered)
+	{
+		bLeftHandleDragged = true;
+		return FReply::Handled().CaptureMouse(AsShared());
+	}
+	else if (bRightHandleHovered)
+	{
+		bRightHandleDragged = true;
+		return FReply::Handled().CaptureMouse(AsShared());
+	}
+
+	return FReply::Unhandled();
+}
+
+FReply STimeRangeSlider::OnMouseButtonUp( const FGeometry& MyGeometry, const FPointerEvent& MouseEvent )
+{
+	ResetState();
+
+	return FReply::Handled().ReleaseMouseCapture();
+}
+
+FReply STimeRangeSlider::OnMouseMove( const FGeometry& MyGeometry, const FPointerEvent& MouseEvent )
+{
+	if (HasMouseCapture())
+	{
+		float DragDelta = ComputeDragDelta(MouseEvent, MyGeometry.Size.X);
+
+		if (bHandleDragged)
+		{
+			if (TimeSliderController.IsValid())
+			{
+				float NewIn = MouseDownViewRange.GetLowerBoundValue() + DragDelta;
+				float NewOut = MouseDownViewRange.GetUpperBoundValue() + DragDelta;
+
+				TRange<float> ClampRange = TimeSliderController.Get()->GetClampRange();
+				if (NewIn < ClampRange.GetLowerBoundValue())
+				{
+					NewIn = ClampRange.GetLowerBoundValue();
+					NewOut = NewIn + (MouseDownViewRange.GetUpperBoundValue() - MouseDownViewRange.GetLowerBoundValue());
+				}				
+				else if (NewOut > ClampRange.GetUpperBoundValue())
+				{
+					NewOut = ClampRange.GetUpperBoundValue();
+					NewIn = NewOut - (MouseDownViewRange.GetUpperBoundValue() - MouseDownViewRange.GetLowerBoundValue());
+				}
+				
+				TimeSliderController.Get()->SetViewRange(NewIn, NewOut, EViewRangeInterpolation::Immediate);
+			}
+		}
+		else if (bLeftHandleDragged)
+		{
+			if (TimeSliderController.IsValid())
+			{
+				float NewIn = MouseDownViewRange.GetLowerBoundValue() + DragDelta;
+				TRange<float> ClampRange = TimeSliderController.Get()->GetClampRange();
+				NewIn = FMath::Clamp(NewIn, ClampRange.GetLowerBoundValue(), ClampRange.GetUpperBoundValue());
+				TimeSliderController.Get()->SetViewRange(NewIn, TimeSliderController.Get()->GetViewRange().GetUpperBoundValue(), EViewRangeInterpolation::Immediate);
+			}
+		}
+		else if (bRightHandleDragged)
+		{
+			if (TimeSliderController.IsValid())
+			{
+				float NewOut = MouseDownViewRange.GetUpperBoundValue() + DragDelta;
+				TRange<float> ClampRange = TimeSliderController.Get()->GetClampRange();
+				NewOut = FMath::Clamp(NewOut, ClampRange.GetLowerBoundValue(), ClampRange.GetUpperBoundValue());
+				TimeSliderController.Get()->SetViewRange(TimeSliderController.Get()->GetViewRange().GetLowerBoundValue(), NewOut, EViewRangeInterpolation::Immediate);
+			}
+		}
+
+		return FReply::Handled();
+	}
+	else
+	{
+		ResetHoveredState();
+
+		float LeftHandleOffset = 0.f;
+		float HandleOffset = 0.f;
+		float RightHandleOffset = 0.f;
+		ComputeHandleOffsets(LeftHandleOffset, HandleOffset, RightHandleOffset, MyGeometry.Size.X);
+		
+		FGeometry LeftHandleRect = MyGeometry.MakeChild(FVector2D(LeftHandleOffset, 0.f), FVector2D(TimeRangeSliderConstants::HandleSize, TimeRangeSliderConstants::HandleSize));
+		FGeometry RightHandleRect = MyGeometry.MakeChild(FVector2D(RightHandleOffset, 0.f), FVector2D(TimeRangeSliderConstants::HandleSize, TimeRangeSliderConstants::HandleSize));
+		FGeometry HandleRect = MyGeometry.MakeChild(FVector2D(HandleOffset, 0.0f), FVector2D(RightHandleOffset-LeftHandleOffset-TimeRangeSliderConstants::HandleSize, TimeRangeSliderConstants::HandleSize));
+
+		FVector2D LocalMousePosition = MouseEvent.GetScreenSpacePosition();
+
+		if (HandleRect.IsUnderLocation(LocalMousePosition))
+		{
+			bHandleHovered = true;
+		}
+		else if (LeftHandleRect.IsUnderLocation(LocalMousePosition))
+		{
+			bLeftHandleHovered = true;
+		}
+		else if (RightHandleRect.IsUnderLocation(LocalMousePosition))
+		{
+			bRightHandleHovered = true;
+		}
+	}
+	return FReply::Unhandled();
+}
+
+void STimeRangeSlider::OnMouseLeave( const FPointerEvent& MouseEvent )
+{
+	SCompoundWidget::OnMouseLeave(MouseEvent);
+
+	if (!HasMouseCapture())
+	{
+		ResetHoveredState();
+	}
+}
+
+FReply STimeRangeSlider::OnMouseButtonDoubleClick( const FGeometry& MyGeometry, const FPointerEvent& MouseEvent )
+{
+	ResetState();
+
+	OnMouseMove(MyGeometry, MouseEvent);
+
+	if (bHandleHovered)
+	{
+		if (TimeSliderController.IsValid())
+		{
+			if (FMath::IsNearlyEqual(TimeSliderController.Get()->GetViewRange().GetLowerBoundValue(), TimeSliderController.Get()->GetClampRange().GetLowerBoundValue()) &&
+				FMath::IsNearlyEqual(TimeSliderController.Get()->GetViewRange().GetUpperBoundValue(), TimeSliderController.Get()->GetClampRange().GetUpperBoundValue()))
+			{
+				if (!LastViewRange.IsEmpty())
+				{
+					TimeSliderController.Get()->SetViewRange(LastViewRange.GetLowerBoundValue(), LastViewRange.GetUpperBoundValue(), EViewRangeInterpolation::Immediate);
+				}
+			}
+			else
+			{
+				LastViewRange = TRange<float>(TimeSliderController.Get()->GetViewRange().GetLowerBoundValue(), TimeSliderController.Get()->GetViewRange().GetUpperBoundValue());
+				TimeSliderController.Get()->SetViewRange(TimeSliderController.Get()->GetClampRange().GetLowerBoundValue(), TimeSliderController.Get()->GetClampRange().GetUpperBoundValue(), EViewRangeInterpolation::Immediate);
+			}
+		}
+
+		ResetState();
+		return FReply::Handled();
+	}
+	ResetState();
+	return FReply::Unhandled();
+}
+
+void STimeRangeSlider::ResetState()
+{
+	bHandleDragged = false;
+	bLeftHandleDragged = false;
+	bRightHandleDragged = false;
+	ResetHoveredState();
+}
+
+void STimeRangeSlider::ResetHoveredState()
+{
+	bHandleHovered = false;
+	bLeftHandleHovered = false;
+	bRightHandleHovered = false;
+}
+
+#undef LOCTEXT_NAMESPACE
diff --git a/Engine/Source/Editor/SequencerWidgets/Private/STimeRangeSlider.h b/Engine/Source/Editor/SequencerWidgets/Private/STimeRangeSlider.h
new file mode 100644
--- /dev/null
+++ b/Engine/Source/Editor/SequencerWidgets/Private/STimeRangeSlider.h
@@ -0,0 +1,54 @@
+// Copyright 1998-2015 Epic Games, Inc. All Rights Reserved.
+
+#pragma once
+
+class STimeRangeSlider : public SCompoundWidget
+{
+public:
+
+	SLATE_BEGIN_ARGS(STimeRangeSlider){}
+		SLATE_DEFAULT_SLOT(FArguments, Content)
+	SLATE_END_ARGS()
+
+	void Construct( const FArguments& InArgs, TSharedRef<ITimeSliderController> InTimeSliderController );
+
+	// SWidget interface
+	virtual int32 OnPaint( const FPaintArgs& Args, const FGeometry& AllottedGeometry, const FSlateRect& MyClippingRect, FSlateWindowElementList& OutDrawElements, int32 LayerId, const FWidgetStyle& InWidgetStyle, bool bParentEnabled ) const override;
+	virtual FReply OnMouseButtonDown( const FGeometry& MyGeometry, const FPointerEvent& MouseEvent ) override;
+	virtual FReply OnMouseButtonUp( const FGeometry& MyGeometry, const FPointerEvent& MouseEvent ) override;
+	virtual FReply OnMouseMove( const FGeometry& MyGeometry, const FPointerEvent& MouseEvent ) override;
+	virtual void OnMouseLeave( const FPointerEvent& MouseEvent ) override;
+	virtual FReply OnMouseButtonDoubleClick( const FGeometry& MyGeometry, const FPointerEvent& MouseEvent ) override;
+
+protected:
+	void ResetState();
+	void ResetHoveredState();
+	float ComputeDragDelta(const FPointerEvent& MouseEvent, int32 GeometryWidth) const;
+	void ComputeHandleOffsets(float& LeftHandleOffset, float& RightHandleOffset, float&HandleOffset, int32 GeometryWidth) const;
+
+private:
+	/* The left handle is being dragged */
+	bool bLeftHandleDragged;
+	/* The right handle is being dragged */
+	bool bRightHandleDragged;
+	/* The handle is being dragged */
+	bool bHandleDragged;
+
+	/* The left handle is hovered */
+	bool bLeftHandleHovered;
+	/* The right handle is hovered */
+	bool bRightHandleHovered;
+	/* The handle is hovered */
+	bool bHandleHovered;
+
+	/* The position of the mouse on mouse down */
+	FVector2D MouseDownPosition;
+
+	/* The in/out view range on mouse down */
+	TRange<float> MouseDownViewRange;
+
+	/* The in/out view range viewed before expansion */
+	TRange<float> LastViewRange;
+
+	TSharedPtr<ITimeSliderController> TimeSliderController;
+};
\ No newline at end of file
diff --git a/Engine/Source/Editor/SequencerWidgets/Private/SequencerWidgetsModule.cpp b/Engine/Source/Editor/SequencerWidgets/Private/SequencerWidgetsModule.cpp
--- a/Engine/Source/Editor/SequencerWidgets/Private/SequencerWidgetsModule.cpp
+++ b/Engine/Source/Editor/SequencerWidgets/Private/SequencerWidgetsModule.cpp
@@ -2,6 +2,7 @@
 
 #include "SequencerWidgetsPrivatePCH.h"
 #include "STimeSlider.h"
+#include "STimeRange.h"
 #include "ModuleManager.h"
 
 
@@ -24,3 +25,19 @@ TSharedRef<ITimeSlider> FSequencerWidgetsModule::CreateTimeSlider( const TShared
 		.MirrorLabels( bMirrorLabels );
 }
 
+TSharedRef<ITimeSlider> FSequencerWidgetsModule::CreateTimeSlider( const TSharedRef<ITimeSliderController>& InController, const TAttribute<EVisibility>& VisibilityDelegate, bool bMirrorLabels )
+{
+	return
+		SNew( STimeSlider, InController )
+		.Visibility(VisibilityDelegate)
+		.MirrorLabels( bMirrorLabels );
+}
+
+TSharedRef<ITimeSlider> FSequencerWidgetsModule::CreateTimeRange( const TSharedRef<ITimeSliderController>& InController, const TAttribute<EVisibility>& VisibilityDelegate, const TAttribute<bool>& ShowFrameNumbersDelegate )
+{
+	return 
+		SNew( STimeRange, InController)
+		.Visibility(VisibilityDelegate)
+		.ShowFrameNumbers(ShowFrameNumbersDelegate);
+}
+
diff --git a/Engine/Source/Editor/SequencerWidgets/Public/ITimeSlider.h b/Engine/Source/Editor/SequencerWidgets/Public/ITimeSlider.h
--- a/Engine/Source/Editor/SequencerWidgets/Public/ITimeSlider.h
+++ b/Engine/Source/Editor/SequencerWidgets/Public/ITimeSlider.h
@@ -15,6 +15,7 @@ enum class EViewRangeInterpolation
 
 DECLARE_DELEGATE_TwoParams( FOnScrubPositionChanged, float, bool )
 DECLARE_DELEGATE_TwoParams( FOnViewRangeChanged, TRange<float>, EViewRangeInterpolation )
+DECLARE_DELEGATE_OneParam( FOnClampRangeChanged, TRange<float> )
 
 /** Structure used to wrap up a range, and an optional animation target */
 struct FAnimatedRange : public TRange<float>
@@ -55,12 +56,12 @@ struct FTimeSliderArgs
 	FTimeSliderArgs()
 		: ScrubPosition(0)
 		, ViewRange( FAnimatedRange(0.0f, 5.0f) )
-		, ClampMin(-70000.0f)
-		, ClampMax(70000.0f)
+		, ClampRange( FAnimatedRange(-FLT_MAX/2.f, FLT_MAX/2.f) )
 		, OnScrubPositionChanged()
 		, OnBeginScrubberMovement()
 		, OnEndScrubberMovement()
 		, OnViewRangeChanged()
+		, OnClampRangeChanged()
 		, AllowZoom(true)
 		, Settings(nullptr)
 	{}
@@ -69,10 +70,8 @@ struct FTimeSliderArgs
 	TAttribute<float> ScrubPosition;
 	/** View time range */
 	TAttribute< FAnimatedRange > ViewRange;
-	/** Optional min output to clamp to */
-	TAttribute<TOptional<float> > ClampMin;
-	/** Optional max output to clamp to */
-	TAttribute<TOptional<float> > ClampMax;
+	/** Clamp time range */
+	TAttribute< FAnimatedRange > ClampRange;
 	/** Called when the scrub position changes */
 	FOnScrubPositionChanged OnScrubPositionChanged;
 	/** Called right before the scrubber begins to move */
@@ -81,6 +80,8 @@ struct FTimeSliderArgs
 	FSimpleDelegate OnEndScrubberMovement;
 	/** Called when the view range changes */
 	FOnViewRangeChanged OnViewRangeChanged;
+	/** Called when the clamp range changes */
+	FOnClampRangeChanged OnClampRangeChanged;
 	/** If we are allowed to zoom */
 	bool AllowZoom;
 	/** User-supplied settings object */
@@ -96,6 +97,35 @@ public:
 	virtual FReply OnMouseButtonUp( TSharedRef<SWidget> WidgetOwner, const FGeometry& MyGeometry, const FPointerEvent& MouseEvent ) = 0;
 	virtual FReply OnMouseMove( TSharedRef<SWidget> WidgetOwner, const FGeometry& MyGeometry, const FPointerEvent& MouseEvent ) = 0;
 	virtual FReply OnMouseWheel( TSharedRef<SWidget> WidgetOwner, const FGeometry& MyGeometry, const FPointerEvent& MouseEvent ) = 0;
+
+	/** Get the current view range for this controller */
+	virtual FAnimatedRange GetViewRange() const { return FAnimatedRange(); }
+
+	/** Get the current clamp range for this controller */
+	virtual FAnimatedRange GetClampRange() const { return FAnimatedRange(); }
+
+	/** Convert a time to a frame */
+	virtual int32 TimeToFrame(float Time) const { return 1; }
+
+	/** Convert a time to a frame */
+	virtual float FrameToTime(int32 Frame) const { return 1.f; }
+
+	/**
+	 * Set a new range based on a min, max and an interpolation mode
+	 * 
+	 * @param NewRangeMin		The new lower bound of the range
+	 * @param NewRangeMax		The new upper bound of the range
+	 * @param Interpolation		How to set the new range (either immediately, or animated)
+	 */
+	virtual void SetViewRange( float NewRangeMin, float NewRangeMax, EViewRangeInterpolation Interpolation ) {}
+
+	/**
+	 * Set a new clamp range based on a min, max
+	 * 
+	 * @param NewRangeMin		The new lower bound of the clamp range
+	 * @param NewRangeMax		The new upper bound of the clamp range
+	 */
+	virtual void SetClampRange( float NewRangeMin, float NewRangeMax) {}
 };
 
 /**
@@ -104,4 +134,7 @@ public:
 class ITimeSlider : public SCompoundWidget
 {
 public:
+	SLATE_BEGIN_ARGS(ITimeSlider){}
+		SLATE_DEFAULT_SLOT(FArguments, Content)
+	SLATE_END_ARGS()
 };
\ No newline at end of file
diff --git a/Engine/Source/Editor/SequencerWidgets/Public/SequencerWidgetsModule.h b/Engine/Source/Editor/SequencerWidgets/Public/SequencerWidgetsModule.h
--- a/Engine/Source/Editor/SequencerWidgets/Public/SequencerWidgetsModule.h
+++ b/Engine/Source/Editor/SequencerWidgets/Public/SequencerWidgetsModule.h
@@ -17,6 +17,8 @@ public:
 
 	virtual void ShutdownModule() override;
 
-	virtual TSharedRef<ITimeSlider> CreateTimeSlider( const TSharedRef<class ITimeSliderController>& InController, bool bMirrorLabels );
+	virtual TSharedRef<ITimeSlider> CreateTimeSlider( const TSharedRef<class ITimeSliderController>& InController, bool bMirrorLabels  );
+	virtual TSharedRef<ITimeSlider> CreateTimeSlider( const TSharedRef<class ITimeSliderController>& InController, const TAttribute<EVisibility>& VisibilityDelegate, bool bMirrorLabels  );
+	virtual TSharedRef<ITimeSlider> CreateTimeRange( const TSharedRef<class ITimeSliderController>& InController, const TAttribute<EVisibility>& VisibilityDelegate, const TAttribute<bool>& ShowFrameNumbersDelegate );
 };
 
diff --git a/Engine/Source/Editor/SequencerWidgets/SequencerWidgets.Build.cs b/Engine/Source/Editor/SequencerWidgets/SequencerWidgets.Build.cs
--- a/Engine/Source/Editor/SequencerWidgets/SequencerWidgets.Build.cs
+++ b/Engine/Source/Editor/SequencerWidgets/SequencerWidgets.Build.cs
@@ -12,6 +12,7 @@ public class SequencerWidgets : ModuleRules
 			new string[] {
 				"Core",
 				"CoreUObject",
+                "EditorStyle",
 				"Engine",
                 "InputCore",
 				"Slate",
diff --git a/Engine/Source/Editor/UMGEditor/Private/TabFactory/AnimationTabSummoner.cpp b/Engine/Source/Editor/UMGEditor/Private/TabFactory/AnimationTabSummoner.cpp
--- a/Engine/Source/Editor/UMGEditor/Private/TabFactory/AnimationTabSummoner.cpp
+++ b/Engine/Source/Editor/UMGEditor/Private/TabFactory/AnimationTabSummoner.cpp
@@ -318,10 +318,17 @@ private:
 
 	FReply OnNewAnimationClicked()
 	{
+		const float InTime = 0.f;
+		const float OutTime = 5.0f;
+
 		UWidgetBlueprint* WidgetBlueprint = BlueprintEditor.Pin()->GetWidgetBlueprintObj();
 
 		UWidgetAnimation* NewAnimation = NewObject<UWidgetAnimation>(WidgetBlueprint, MakeUniqueObjectName(WidgetBlueprint, UWidgetAnimation::StaticClass(), "NewAnimation"), RF_Transactional);
 		NewAnimation->MovieScene = NewObject<UMovieScene>(NewAnimation, NewAnimation->GetFName(), RF_Transactional);
+		NewAnimation->MovieScene->StartTime = InTime;
+		NewAnimation->MovieScene->InTime = InTime;
+		NewAnimation->MovieScene->OutTime = OutTime;
+		NewAnimation->MovieScene->EndTime = OutTime;
 
 		bool bRequestRename = true;
 		bool bNewAnimation = true;
diff --git a/Engine/Source/Editor/UMGEditor/Private/WidgetBlueprintEditor.cpp b/Engine/Source/Editor/UMGEditor/Private/WidgetBlueprintEditor.cpp
--- a/Engine/Source/Editor/UMGEditor/Private/WidgetBlueprintEditor.cpp
+++ b/Engine/Source/Editor/UMGEditor/Private/WidgetBlueprintEditor.cpp
@@ -548,17 +548,26 @@ TSharedPtr<ISequencer>& FWidgetBlueprintEditor::GetSequencer()
 {
 	if(!Sequencer.IsValid())
 	{
+		const float InTime = 0.f;
+		const float OutTime = 5.0f;
+
 		FSequencerViewParams ViewParams(TEXT("UMGSequencerSettings"));
 		{
-			ViewParams.InitalViewRange = TRange<float>(-0.02f, 3.2f);
+			ViewParams.InitalViewRange = TRange<float>(InTime, OutTime);
 			ViewParams.InitialScrubPosition = 0;
 			ViewParams.OnGetAddMenuContent = FOnGetAddMenuContent::CreateSP(this, &FWidgetBlueprintEditor::OnGetAnimationAddMenuContent);
 		}
 
 		FSequencerInitParams SequencerInitParams;
 		{
+			UWidgetAnimation* NullAnimation = UWidgetAnimation::GetNullAnimation();
+			NullAnimation->MovieScene->StartTime = InTime;
+			NullAnimation->MovieScene->InTime = InTime;
+			NullAnimation->MovieScene->OutTime = OutTime;
+			NullAnimation->MovieScene->EndTime = OutTime;
+
 			SequencerInitParams.ViewParams = ViewParams;
-			SequencerInitParams.Animation = UWidgetAnimation::GetNullAnimation();
+			SequencerInitParams.Animation = NullAnimation;
 			SequencerInitParams.bEditWithinLevelEditor = false;
 			SequencerInitParams.ToolkitHost = nullptr;
 		};
diff --git a/Engine/Source/Runtime/MovieScene/Private/MovieScene.cpp b/Engine/Source/Runtime/MovieScene/Private/MovieScene.cpp
--- a/Engine/Source/Runtime/MovieScene/Private/MovieScene.cpp
+++ b/Engine/Source/Runtime/MovieScene/Private/MovieScene.cpp
@@ -7,6 +7,15 @@
 /* UMovieScene interface
  *****************************************************************************/
 
+UMovieScene::UMovieScene(const FObjectInitializer& ObjectInitializer)
+	: Super(ObjectInitializer)
+	, InTime(FLT_MAX)
+	, OutTime(-FLT_MAX)
+	, StartTime(FLT_MAX)
+	, EndTime(-FLT_MAX)
+{ }
+
+
 #if WITH_EDITOR
 
 // @todo sequencer: Some of these methods should only be used by tools, and should probably move out of MovieScene!
@@ -181,20 +190,7 @@ FMovieScenePossessable& UMovieScene::GetPossessable( const int32 Index )
 
 TRange<float> UMovieScene::GetTimeRange() const
 {
-	// Get the range of all sections combined
-	TArray< TRange<float> > Bounds;
-
-	for (int32 TypeIndex = 0; TypeIndex < MasterTracks.Num(); ++TypeIndex)
-	{
-		Bounds.Add(MasterTracks[TypeIndex]->GetSectionBoundaries());
-	}
-
-	for (int32 BindingIndex = 0; BindingIndex < ObjectBindings.Num(); ++BindingIndex)
-	{
-		Bounds.Add(ObjectBindings[BindingIndex].GetTimeRange());
-	}
-
-	return TRange<float>::Hull(Bounds);
+	return TRange<float>(InTime, OutTime);
 }
 
 
diff --git a/Engine/Source/Runtime/MovieScene/Public/MovieScene.h b/Engine/Source/Runtime/MovieScene/Public/MovieScene.h
--- a/Engine/Source/Runtime/MovieScene/Public/MovieScene.h
+++ b/Engine/Source/Runtime/MovieScene/Public/MovieScene.h
@@ -49,7 +49,7 @@ UCLASS()
 class MOVIESCENE_API UMovieScene
 	: public UObject
 {
-	GENERATED_BODY()
+	GENERATED_UCLASS_BODY()
 
 public:
 
@@ -251,7 +251,7 @@ public:
 	}
 
 	/**
-	 * @return The time range of the movie scene (defined by all sections in the scene)
+	 * @return The time range of the movie scene (the in to out time range)
 	 */
 	TRange<float> GetTimeRange() const;
 
@@ -312,4 +312,21 @@ private:
 	UPROPERTY()
 	FMovieSceneEditorData EditorData;
 #endif
+
+public:
+	/** In time of the movie scene in seconds*/
+	UPROPERTY()
+	float InTime;
+
+	/** Out time of the movie scene in seconds */
+	UPROPERTY()
+	float OutTime;
+
+	/** Start time of the movie scene in seconds*/
+	UPROPERTY()
+	float StartTime;
+
+	/** End time of the movie scene in seconds */
+	UPROPERTY()
+	float EndTime;
 };
