Fix non-unity compile.

diff --git a/Engine/Source/Editor/SequencerWidgets/Private/STimeRangeSlider.cpp b/Engine/Source/Editor/SequencerWidgets/Private/STimeRangeSlider.cpp
--- a/Engine/Source/Editor/SequencerWidgets/Private/STimeRangeSlider.cpp
+++ b/Engine/Source/Editor/SequencerWidgets/Private/STimeRangeSlider.cpp
@@ -3,6 +3,7 @@
 #include "SequencerWidgetsPrivatePCH.h"
 #include "STimeRangeSlider.h"
 #include "SlateStyle.h"
+#include "EditorStyle.h"
 
 #define LOCTEXT_NAMESPACE "STimeRangeSlider"
 
