Sequencer: Ensure out frame is not less than in frame and vice versa.

diff --git a/Engine/Source/Editor/Sequencer/Private/SSequencer.cpp b/Engine/Source/Editor/Sequencer/Private/SSequencer.cpp
--- a/Engine/Source/Editor/Sequencer/Private/SSequencer.cpp
+++ b/Engine/Source/Editor/Sequencer/Private/SSequencer.cpp
@@ -214,7 +214,7 @@ void SSequencer::Construct( const FArguments& InArgs, TSharedRef< class FSequenc
 	TSharedRef<ITimeSlider> BottomTimeSlider = SequencerWidgets.CreateTimeSlider( TimeSliderController, TAttribute<EVisibility>(this, &SSequencer::GetBottomTimeSliderVisibility), bMirrorLabels );
 
 	// Create bottom time range slider
-	TSharedRef<ITimeSlider> BottomTimeRange = SequencerWidgets.CreateTimeRange( TimeSliderController, TAttribute<EVisibility>(this, &SSequencer::GetTimeRangeVisibility), TAttribute<bool>(this, &SSequencer::ShowFrameNumbers));
+	TSharedRef<ITimeSlider> BottomTimeRange = SequencerWidgets.CreateTimeRange( TimeSliderController, TAttribute<EVisibility>(this, &SSequencer::GetTimeRangeVisibility), TAttribute<bool>(this, &SSequencer::ShowFrameNumbers), TAttribute<float>(this, &SSequencer::OnGetTimeSnapInterval));
 
 	OnGetAddMenuContent = InArgs._OnGetAddMenuContent;
 
diff --git a/Engine/Source/Editor/SequencerWidgets/Private/STimeRange.cpp b/Engine/Source/Editor/SequencerWidgets/Private/STimeRange.cpp
--- a/Engine/Source/Editor/SequencerWidgets/Private/STimeRange.cpp
+++ b/Engine/Source/Editor/SequencerWidgets/Private/STimeRange.cpp
@@ -22,6 +22,7 @@ void STimeRange::Construct( const STimeRange::FArguments& InArgs, TSharedRef<ITi
 {
 	TimeSliderController = InTimeSliderController;
 	ShowFrameNumbers = InArgs._ShowFrameNumbers;
+	TimeSnapInterval = InArgs._TimeSnapInterval;
 	
 	this->ChildSlot
 	.HAlign(HAlign_Fill)
@@ -50,7 +51,7 @@ void STimeRange::Construct( const STimeRange::FArguments& InArgs, TSharedRef<ITi
 		.VAlign(VAlign_Fill)
 		.Padding(2.0f, 4.0f)
 		[
-			SAssignNew(TimeRangeSlider, STimeRangeSlider, InTimeSliderController)
+			SAssignNew(TimeRangeSlider, STimeRangeSlider, InTimeSliderController, TSharedRef<STimeRange>(this))
 		]
 		+SHorizontalBox::Slot()
 		.AutoWidth()
@@ -324,4 +325,13 @@ FText STimeRange::EndTimeTooltip() const
 	}
 }
 
+float STimeRange::GetTimeSnapInterval() const
+{
+	if (TimeSnapInterval.IsBound())
+	{
+		return TimeSnapInterval.Get();
+	}
+	return 1.f;
+}
+
 #undef LOCTEXT_NAMESPACE
diff --git a/Engine/Source/Editor/SequencerWidgets/Private/STimeRange.h b/Engine/Source/Editor/SequencerWidgets/Private/STimeRange.h
--- a/Engine/Source/Editor/SequencerWidgets/Private/STimeRange.h
+++ b/Engine/Source/Editor/SequencerWidgets/Private/STimeRange.h
@@ -11,14 +11,12 @@ DECLARE_DELEGATE_RetVal( bool, STimeRangeGetter );
 class STimeRange : public ITimeSlider
 {
 public:
-
-	//SLATE_BEGIN_ARGS(STimeRange){}
-	//	SLATE_DEFAULT_SLOT(FArguments, Content)
-	//SLATE_END_ARGS()
 	SLATE_BEGIN_ARGS(STimeRange)
 	{}
 		/* If we should show frame numbers on the timeline */
 		SLATE_ARGUMENT( TAttribute<bool>, ShowFrameNumbers )
+		/* The time snap interval for the timeline */
+		SLATE_ARGUMENT( TAttribute<float>, TimeSnapInterval )
 	SLATE_END_ARGS()
 
 	/**
@@ -28,6 +26,8 @@ public:
 	 */
 	void Construct( const FArguments& InArgs, TSharedRef<ITimeSliderController> InTimeSliderController );
 
+	float GetTimeSnapInterval() const;
+
 protected:
 	FText InTime() const;
 	FText OutTime() const;
@@ -48,4 +48,5 @@ private:
 	TSharedPtr<ITimeSliderController> TimeSliderController;
 	TSharedPtr<STimeRangeSlider> TimeRangeSlider;
 	TAttribute<bool> ShowFrameNumbers;
+	TAttribute<float> TimeSnapInterval;
 };
\ No newline at end of file
diff --git a/Engine/Source/Editor/SequencerWidgets/Private/STimeRangeSlider.cpp b/Engine/Source/Editor/SequencerWidgets/Private/STimeRangeSlider.cpp
--- a/Engine/Source/Editor/SequencerWidgets/Private/STimeRangeSlider.cpp
+++ b/Engine/Source/Editor/SequencerWidgets/Private/STimeRangeSlider.cpp
@@ -2,6 +2,7 @@
 
 #include "SequencerWidgetsPrivatePCH.h"
 #include "STimeRangeSlider.h"
+#include "STimeRange.h"
 #include "SlateStyle.h"
 #include "EditorStyle.h"
 
@@ -12,9 +13,10 @@ namespace TimeRangeSliderConstants
 	const int32 HandleSize = 14;
 }
 
-void STimeRangeSlider::Construct( const FArguments& InArgs, TSharedRef<ITimeSliderController> InTimeSliderController)
+void STimeRangeSlider::Construct( const FArguments& InArgs, TSharedRef<ITimeSliderController> InTimeSliderController, TSharedRef<STimeRange> InTimeRange)
 {
 	TimeSliderController = InTimeSliderController;
+	TimeRange = InTimeRange;
 	LastViewRange = TimeSliderController.Get()->GetViewRange();
 
 	ResetState();
@@ -74,6 +76,16 @@ int32 STimeRangeSlider::OnPaint( const FPaintArgs& Args, const FGeometry& Allott
 	static const FName SelectionColorName("SelectionColor");
 	FLinearColor SelectionColor = FEditorStyle::GetSlateColor(SelectionColorName).GetColor(FWidgetStyle());
 
+	// Draw the handle box
+	FSlateDrawElement::MakeBox( 
+		OutDrawElements,
+		LayerId, 
+		AllottedGeometry.ToPaintGeometry(FVector2D(HandleOffset, 0.0f), FVector2D(RightHandleOffset-LeftHandleOffset-TimeRangeSliderConstants::HandleSize, TimeRangeSliderConstants::HandleSize)),
+		RangeHandle,
+		MyClippingRect,
+		ESlateDrawEffect::None,
+		(bHandleDragged || bHandleHovered) ? SelectionColor : FLinearColor::Gray);
+
 	// Draw the left handle box
 	FSlateDrawElement::MakeBox( 
 		OutDrawElements,
@@ -94,16 +106,6 @@ int32 STimeRangeSlider::OnPaint( const FPaintArgs& Args, const FGeometry& Allott
 		ESlateDrawEffect::None,
 		(bRightHandleDragged || bRightHandleHovered) ? SelectionColor : FLinearColor::Gray);
 
-	// Draw the handle box
-	FSlateDrawElement::MakeBox( 
-		OutDrawElements,
-		LayerId, 
-		AllottedGeometry.ToPaintGeometry(FVector2D(HandleOffset, 0.0f), FVector2D(RightHandleOffset-LeftHandleOffset-TimeRangeSliderConstants::HandleSize, TimeRangeSliderConstants::HandleSize)),
-		RangeHandle,
-		MyClippingRect,
-		ESlateDrawEffect::None,
-		(bHandleDragged || bHandleHovered) ? SelectionColor : FLinearColor::Gray);
-
 	SCompoundWidget::OnPaint(Args, AllottedGeometry, MyClippingRect, OutDrawElements, LayerId, InWidgetStyle, ShouldBeEnabled( bParentEnabled ));
 
 	return LayerId;
@@ -178,6 +180,11 @@ FReply STimeRangeSlider::OnMouseMove( const FGeometry& MyGeometry, const FPointe
 				float NewIn = MouseDownViewRange.GetLowerBoundValue() + DragDelta;
 				TRange<float> ClampRange = TimeSliderController.Get()->GetClampRange();
 				NewIn = FMath::Clamp(NewIn, ClampRange.GetLowerBoundValue(), ClampRange.GetUpperBoundValue());
+				if (NewIn > TimeSliderController.Get()->GetViewRange().GetUpperBoundValue())
+				{
+					NewIn = TimeSliderController.Get()->GetViewRange().GetUpperBoundValue() - TimeRange.Get()->GetTimeSnapInterval();
+				}
+
 				TimeSliderController.Get()->SetViewRange(NewIn, TimeSliderController.Get()->GetViewRange().GetUpperBoundValue(), EViewRangeInterpolation::Immediate);
 			}
 		}
@@ -188,6 +195,11 @@ FReply STimeRangeSlider::OnMouseMove( const FGeometry& MyGeometry, const FPointe
 				float NewOut = MouseDownViewRange.GetUpperBoundValue() + DragDelta;
 				TRange<float> ClampRange = TimeSliderController.Get()->GetClampRange();
 				NewOut = FMath::Clamp(NewOut, ClampRange.GetLowerBoundValue(), ClampRange.GetUpperBoundValue());
+				if (NewOut < TimeSliderController.Get()->GetViewRange().GetLowerBoundValue())
+				{
+					NewOut = TimeSliderController.Get()->GetViewRange().GetLowerBoundValue() + TimeRange.Get()->GetTimeSnapInterval();
+				}
+
 				TimeSliderController.Get()->SetViewRange(TimeSliderController.Get()->GetViewRange().GetLowerBoundValue(), NewOut, EViewRangeInterpolation::Immediate);
 			}
 		}
diff --git a/Engine/Source/Editor/SequencerWidgets/Private/STimeRangeSlider.h b/Engine/Source/Editor/SequencerWidgets/Private/STimeRangeSlider.h
--- a/Engine/Source/Editor/SequencerWidgets/Private/STimeRangeSlider.h
+++ b/Engine/Source/Editor/SequencerWidgets/Private/STimeRangeSlider.h
@@ -2,15 +2,16 @@
 
 #pragma once
 
+class STimeRange;
+
 class STimeRangeSlider : public SCompoundWidget
 {
 public:
-
 	SLATE_BEGIN_ARGS(STimeRangeSlider){}
 		SLATE_DEFAULT_SLOT(FArguments, Content)
 	SLATE_END_ARGS()
 
-	void Construct( const FArguments& InArgs, TSharedRef<ITimeSliderController> InTimeSliderController );
+	void Construct( const FArguments& InArgs, TSharedRef<ITimeSliderController> InTimeSliderController, TSharedRef<STimeRange> InTimeRange );
 
 	// SWidget interface
 	virtual int32 OnPaint( const FPaintArgs& Args, const FGeometry& AllottedGeometry, const FSlateRect& MyClippingRect, FSlateWindowElementList& OutDrawElements, int32 LayerId, const FWidgetStyle& InWidgetStyle, bool bParentEnabled ) const override;
@@ -51,4 +52,5 @@ private:
 	TRange<float> LastViewRange;
 
 	TSharedPtr<ITimeSliderController> TimeSliderController;
+	TSharedPtr<STimeRange> TimeRange;
 };
\ No newline at end of file
diff --git a/Engine/Source/Editor/SequencerWidgets/Private/SequencerWidgetsModule.cpp b/Engine/Source/Editor/SequencerWidgets/Private/SequencerWidgetsModule.cpp
--- a/Engine/Source/Editor/SequencerWidgets/Private/SequencerWidgetsModule.cpp
+++ b/Engine/Source/Editor/SequencerWidgets/Private/SequencerWidgetsModule.cpp
@@ -33,11 +33,12 @@ TSharedRef<ITimeSlider> FSequencerWidgetsModule::CreateTimeSlider( const TShared
 		.MirrorLabels( bMirrorLabels );
 }
 
-TSharedRef<ITimeSlider> FSequencerWidgetsModule::CreateTimeRange( const TSharedRef<ITimeSliderController>& InController, const TAttribute<EVisibility>& VisibilityDelegate, const TAttribute<bool>& ShowFrameNumbersDelegate )
+TSharedRef<ITimeSlider> FSequencerWidgetsModule::CreateTimeRange( const TSharedRef<ITimeSliderController>& InController, const TAttribute<EVisibility>& VisibilityDelegate, const TAttribute<bool>& ShowFrameNumbersDelegate, const TAttribute<float>& TimeSnapIntervalDelegate )
 {
 	return 
 		SNew( STimeRange, InController)
 		.Visibility(VisibilityDelegate)
-		.ShowFrameNumbers(ShowFrameNumbersDelegate);
+		.ShowFrameNumbers(ShowFrameNumbersDelegate)
+		.TimeSnapInterval(TimeSnapIntervalDelegate);
 }
 
diff --git a/Engine/Source/Editor/SequencerWidgets/Public/SequencerWidgetsModule.h b/Engine/Source/Editor/SequencerWidgets/Public/SequencerWidgetsModule.h
--- a/Engine/Source/Editor/SequencerWidgets/Public/SequencerWidgetsModule.h
+++ b/Engine/Source/Editor/SequencerWidgets/Public/SequencerWidgetsModule.h
@@ -19,6 +19,6 @@ public:
 
 	virtual TSharedRef<ITimeSlider> CreateTimeSlider( const TSharedRef<class ITimeSliderController>& InController, bool bMirrorLabels  );
 	virtual TSharedRef<ITimeSlider> CreateTimeSlider( const TSharedRef<class ITimeSliderController>& InController, const TAttribute<EVisibility>& VisibilityDelegate, bool bMirrorLabels  );
-	virtual TSharedRef<ITimeSlider> CreateTimeRange( const TSharedRef<class ITimeSliderController>& InController, const TAttribute<EVisibility>& VisibilityDelegate, const TAttribute<bool>& ShowFrameNumbersDelegate );
+	virtual TSharedRef<ITimeSlider> CreateTimeRange( const TSharedRef<class ITimeSliderController>& InController, const TAttribute<EVisibility>& VisibilityDelegate, const TAttribute<bool>& ShowFrameNumbersDelegate, const TAttribute<float>& TimeSnapIntervalDelegate );
 };
 
