Sequencer: Fix time range slider crash on exit.

diff --git a/Engine/Source/Editor/SequencerWidgets/Private/STimeRange.cpp b/Engine/Source/Editor/SequencerWidgets/Private/STimeRange.cpp
--- a/Engine/Source/Editor/SequencerWidgets/Private/STimeRange.cpp
+++ b/Engine/Source/Editor/SequencerWidgets/Private/STimeRange.cpp
@@ -51,7 +51,7 @@ void STimeRange::Construct( const STimeRange::FArguments& InArgs, TSharedRef<ITi
 		.VAlign(VAlign_Fill)
 		.Padding(2.0f, 4.0f)
 		[
-			SAssignNew(TimeRangeSlider, STimeRangeSlider, InTimeSliderController, TSharedRef<STimeRange>(this))
+			SAssignNew(TimeRangeSlider, STimeRangeSlider, InTimeSliderController, SharedThis(this))
 		]
 		+SHorizontalBox::Slot()
 		.AutoWidth()
diff --git a/Engine/Source/Editor/SequencerWidgets/Private/STimeRangeSlider.cpp b/Engine/Source/Editor/SequencerWidgets/Private/STimeRangeSlider.cpp
--- a/Engine/Source/Editor/SequencerWidgets/Private/STimeRangeSlider.cpp
+++ b/Engine/Source/Editor/SequencerWidgets/Private/STimeRangeSlider.cpp
@@ -13,7 +13,7 @@ namespace TimeRangeSliderConstants
 	const int32 HandleSize = 14;
 }
 
-void STimeRangeSlider::Construct( const FArguments& InArgs, TSharedRef<ITimeSliderController> InTimeSliderController, TSharedRef<STimeRange> InTimeRange)
+void STimeRangeSlider::Construct( const FArguments& InArgs, TSharedRef<ITimeSliderController> InTimeSliderController, TSharedPtr<STimeRange> InTimeRange)
 {
 	TimeSliderController = InTimeSliderController;
 	TimeRange = InTimeRange;
@@ -182,7 +182,10 @@ FReply STimeRangeSlider::OnMouseMove( const FGeometry& MyGeometry, const FPointe
 				NewIn = FMath::Clamp(NewIn, ClampRange.GetLowerBoundValue(), ClampRange.GetUpperBoundValue());
 				if (NewIn > TimeSliderController.Get()->GetViewRange().GetUpperBoundValue())
 				{
-					NewIn = TimeSliderController.Get()->GetViewRange().GetUpperBoundValue() - TimeRange.Get()->GetTimeSnapInterval();
+					if (TimeRange.IsValid())
+					{
+						NewIn = TimeSliderController.Get()->GetViewRange().GetUpperBoundValue() - TimeRange.Pin()->GetTimeSnapInterval();
+					}
 				}
 
 				TimeSliderController.Get()->SetViewRange(NewIn, TimeSliderController.Get()->GetViewRange().GetUpperBoundValue(), EViewRangeInterpolation::Immediate);
@@ -197,7 +200,10 @@ FReply STimeRangeSlider::OnMouseMove( const FGeometry& MyGeometry, const FPointe
 				NewOut = FMath::Clamp(NewOut, ClampRange.GetLowerBoundValue(), ClampRange.GetUpperBoundValue());
 				if (NewOut < TimeSliderController.Get()->GetViewRange().GetLowerBoundValue())
 				{
-					NewOut = TimeSliderController.Get()->GetViewRange().GetLowerBoundValue() + TimeRange.Get()->GetTimeSnapInterval();
+					if (TimeRange.IsValid())
+					{
+						NewOut = TimeSliderController.Get()->GetViewRange().GetLowerBoundValue() + TimeRange.Pin()->GetTimeSnapInterval();
+					}
 				}
 
 				TimeSliderController.Get()->SetViewRange(TimeSliderController.Get()->GetViewRange().GetLowerBoundValue(), NewOut, EViewRangeInterpolation::Immediate);
diff --git a/Engine/Source/Editor/SequencerWidgets/Private/STimeRangeSlider.h b/Engine/Source/Editor/SequencerWidgets/Private/STimeRangeSlider.h
--- a/Engine/Source/Editor/SequencerWidgets/Private/STimeRangeSlider.h
+++ b/Engine/Source/Editor/SequencerWidgets/Private/STimeRangeSlider.h
@@ -11,7 +11,7 @@ public:
 		SLATE_DEFAULT_SLOT(FArguments, Content)
 	SLATE_END_ARGS()
 
-	void Construct( const FArguments& InArgs, TSharedRef<ITimeSliderController> InTimeSliderController, TSharedRef<STimeRange> InTimeRange );
+	void Construct( const FArguments& InArgs, TSharedRef<ITimeSliderController> InTimeSliderController, TSharedPtr<STimeRange> InTimeRange );
 
 	// SWidget interface
 	virtual int32 OnPaint( const FPaintArgs& Args, const FGeometry& AllottedGeometry, const FSlateRect& MyClippingRect, FSlateWindowElementList& OutDrawElements, int32 LayerId, const FWidgetStyle& InWidgetStyle, bool bParentEnabled ) const override;
@@ -52,5 +52,5 @@ private:
 	TRange<float> LastViewRange;
 
 	TSharedPtr<ITimeSliderController> TimeSliderController;
-	TSharedPtr<STimeRange> TimeRange;
+	TWeakPtr<STimeRange> TimeRange;
 };
\ No newline at end of file
