Fixed time range slider in sequencer being able to overlap

The geometry represented by the view range is now a maximum of GeometryWidth - 2* HandleSize. The handles are now represented on the outside of the view range, rather than the inside as was the case previously. This fixes issues with overlapping handles.

diff --git a/Engine/Source/Editor/SequencerWidgets/Private/STimeRangeSlider.cpp b/Engine/Source/Editor/SequencerWidgets/Private/STimeRangeSlider.cpp
--- a/Engine/Source/Editor/SequencerWidgets/Private/STimeRangeSlider.cpp
+++ b/Engine/Source/Editor/SequencerWidgets/Private/STimeRangeSlider.cpp
@@ -35,7 +35,8 @@ float STimeRangeSlider::ComputeDragDelta(const FPointerEvent& MouseEvent, int32
 	}
 	float DragDistance = (MouseEvent.GetScreenSpacePosition() - MouseDownPosition).X;
 
-	return DragDistance / (float)GeometryWidth * (EndTime - StartTime);
+	const float PixelToUnits = (EndTime - StartTime) / (GeometryWidth - TimeRangeSliderConstants::HandleSize*2);
+	return DragDistance * PixelToUnits;
 }
 
 void STimeRangeSlider::ComputeHandleOffsets(float& LeftHandleOffset, float& HandleOffset, float& RightHandleOffset, int32 GeometryWidth) const
@@ -53,9 +54,11 @@ void STimeRangeSlider::ComputeHandleOffsets(float& LeftHandleOffset, float& Hand
 		EndTime = TimeSliderController.Get()->GetClampRange().GetUpperBoundValue();
 	}
 
-	LeftHandleOffset = (InTime - StartTime) / (EndTime - StartTime) * (float)GeometryWidth;
+	const float UnitsToPixel = float(GeometryWidth - TimeRangeSliderConstants::HandleSize*2) / (EndTime - StartTime);
+
+	LeftHandleOffset = (InTime - StartTime) * UnitsToPixel;
 	HandleOffset = LeftHandleOffset + TimeRangeSliderConstants::HandleSize;
-	RightHandleOffset = (OutTime - StartTime) / (EndTime - StartTime) * (float)GeometryWidth - TimeRangeSliderConstants::HandleSize;
+	RightHandleOffset = HandleOffset + (OutTime - InTime) * UnitsToPixel;
 }
 
 int32 STimeRangeSlider::OnPaint( const FPaintArgs& Args, const FGeometry& AllottedGeometry, const FSlateRect& MyClippingRect, FSlateWindowElementList& OutDrawElements, int32 LayerId, const FWidgetStyle& InWidgetStyle, bool bParentEnabled ) const
@@ -151,63 +154,67 @@ FReply STimeRangeSlider::OnMouseMove( const FGeometry& MyGeometry, const FPointe
 	{
 		float DragDelta = ComputeDragDelta(MouseEvent, MyGeometry.Size.X);
 
-		if (bHandleDragged)
+		float SnapInterval = KINDA_SMALL_NUMBER;
 		{
-			if (TimeSliderController.IsValid())
+			auto TimeRangePin = TimeRange.Pin();
+			if (TimeRangePin.IsValid())
 			{
-				float NewIn = MouseDownViewRange.GetLowerBoundValue() + DragDelta;
-				float NewOut = MouseDownViewRange.GetUpperBoundValue() + DragDelta;
+				SnapInterval = TimeRangePin->GetTimeSnapInterval();
+			}
+		}
 
-				TRange<float> ClampRange = TimeSliderController.Get()->GetClampRange();
-				if (NewIn < ClampRange.GetLowerBoundValue())
-				{
-					NewIn = ClampRange.GetLowerBoundValue();
-					NewOut = NewIn + (MouseDownViewRange.GetUpperBoundValue() - MouseDownViewRange.GetLowerBoundValue());
-				}				
-				else if (NewOut > ClampRange.GetUpperBoundValue())
-				{
-					NewOut = ClampRange.GetUpperBoundValue();
-					NewIn = NewOut - (MouseDownViewRange.GetUpperBoundValue() - MouseDownViewRange.GetLowerBoundValue());
-				}
-				
-				TimeSliderController.Get()->SetViewRange(NewIn, NewOut, EViewRangeInterpolation::Immediate);
+		ITimeSliderController* TimeSliderControllerPtr = TimeSliderController.Get();
+		if (!TimeSliderControllerPtr)
+		{
+			return FReply::Handled();
+		}
+
+		if (bHandleDragged)
+		{
+			float NewIn = MouseDownViewRange.GetLowerBoundValue() + DragDelta;
+			float NewOut = MouseDownViewRange.GetUpperBoundValue() + DragDelta;
+
+			TRange<float> ClampRange = TimeSliderControllerPtr->GetClampRange();
+			if (NewIn < ClampRange.GetLowerBoundValue())
+			{
+				NewIn = ClampRange.GetLowerBoundValue();
+				NewOut = NewIn + (MouseDownViewRange.GetUpperBoundValue() - MouseDownViewRange.GetLowerBoundValue());
 			}
+			else if (NewOut > ClampRange.GetUpperBoundValue())
+			{
+				NewOut = ClampRange.GetUpperBoundValue();
+				NewIn = NewOut - (MouseDownViewRange.GetUpperBoundValue() - MouseDownViewRange.GetLowerBoundValue());
+			}
+			
+			TimeSliderControllerPtr->SetViewRange(NewIn, NewOut, EViewRangeInterpolation::Immediate);
 		}
 		else if (bLeftHandleDragged)
 		{
-			if (TimeSliderController.IsValid())
-			{
-				float NewIn = MouseDownViewRange.GetLowerBoundValue() + DragDelta;
-				TRange<float> ClampRange = TimeSliderController.Get()->GetClampRange();
-				NewIn = FMath::Clamp(NewIn, ClampRange.GetLowerBoundValue(), ClampRange.GetUpperBoundValue());
-				if (NewIn > TimeSliderController.Get()->GetViewRange().GetUpperBoundValue())
-				{
-					if (TimeRange.IsValid())
-					{
-						NewIn = TimeSliderController.Get()->GetViewRange().GetUpperBoundValue() - TimeRange.Pin()->GetTimeSnapInterval();
-					}
-				}
+			float NewIn = MouseDownViewRange.GetLowerBoundValue() + DragDelta;
+			float MaxIn = TimeSliderControllerPtr->GetViewRange().GetUpperBoundValue() - SnapInterval;
 
-				TimeSliderController.Get()->SetViewRange(NewIn, TimeSliderController.Get()->GetViewRange().GetUpperBoundValue(), EViewRangeInterpolation::Immediate);
-			}
+			// Clamp to the upper view range bound minus some small number to prevent 0-sized (or negative) ranges
+			NewIn = FMath::Min(NewIn, MaxIn);
+
+			// Double check we're not outside the clamp range
+			TRange<float> ClampRange = TimeSliderControllerPtr->GetClampRange();
+			NewIn = FMath::Clamp(NewIn, ClampRange.GetLowerBoundValue(), ClampRange.GetUpperBoundValue());
+
+			TimeSliderControllerPtr->SetViewRange(NewIn, TimeSliderControllerPtr->GetViewRange().GetUpperBoundValue(), EViewRangeInterpolation::Immediate);
 		}
 		else if (bRightHandleDragged)
 		{
-			if (TimeSliderController.IsValid())
-			{
-				float NewOut = MouseDownViewRange.GetUpperBoundValue() + DragDelta;
-				TRange<float> ClampRange = TimeSliderController.Get()->GetClampRange();
-				NewOut = FMath::Clamp(NewOut, ClampRange.GetLowerBoundValue(), ClampRange.GetUpperBoundValue());
-				if (NewOut < TimeSliderController.Get()->GetViewRange().GetLowerBoundValue())
-				{
-					if (TimeRange.IsValid())
-					{
-						NewOut = TimeSliderController.Get()->GetViewRange().GetLowerBoundValue() + TimeRange.Pin()->GetTimeSnapInterval();
-					}
-				}
+			float NewOut = MouseDownViewRange.GetUpperBoundValue() + DragDelta;
+			float MinOut = TimeSliderControllerPtr->GetViewRange().GetLowerBoundValue() + SnapInterval;
 
-				TimeSliderController.Get()->SetViewRange(TimeSliderController.Get()->GetViewRange().GetLowerBoundValue(), NewOut, EViewRangeInterpolation::Immediate);
-			}
+			// Clamp to the upper view range bound minus some small number to prevent 0-sized (or negative) ranges
+			NewOut = FMath::Max(NewOut, MinOut);
+
+			// Double check we're not outside the clamp range
+			TRange<float> ClampRange = TimeSliderControllerPtr->GetClampRange();
+			NewOut = FMath::Clamp(NewOut, ClampRange.GetLowerBoundValue(), ClampRange.GetUpperBoundValue());
+
+			TimeSliderControllerPtr->SetViewRange(TimeSliderController.Get()->GetViewRange().GetLowerBoundValue(), NewOut, EViewRangeInterpolation::Immediate);
 		}
 
 		return FReply::Handled();
