Changed spinbox styles for numeric inputs on the sequencer UI

Also made some minor improvements to the time range slider to ensure it cannot end up with mismatched view/clamp ranges.

diff --git a/Engine/Source/Editor/EditorStyle/Private/SlateEditorStyle.cpp b/Engine/Source/Editor/EditorStyle/Private/SlateEditorStyle.cpp
--- a/Engine/Source/Editor/EditorStyle/Private/SlateEditorStyle.cpp
+++ b/Engine/Source/Editor/EditorStyle/Private/SlateEditorStyle.cpp
@@ -1645,9 +1645,15 @@ void FSlateEditorStyle::FStyle::SetupGeneralStyles()
 			.SetHandleHighlightBrush( FSlateNoResource() );
 		Set( "Sequencer.AnimationOutliner.Splitter", OutlinerSplitterStyle );
 
-		const FSpinBoxStyle KeyEditorSpinBoxStyle = FSpinBoxStyle(GetWidgetStyle<FSpinBoxStyle>("SpinBox"))
-			.SetTextPadding(FMargin(2, 0));
-		Set( "Sequencer.AnimationOutliner.KeyEditorSpinBoxStyle", KeyEditorSpinBoxStyle );
+		Set( "Sequencer.HyperlinkSpinBox", FSpinBoxStyle(GetWidgetStyle<FSpinBoxStyle>("SpinBox"))
+			.SetTextPadding(FMargin(0))
+			.SetBackgroundBrush(BORDER_BRUSH("Old/HyperlinkDotted", FMargin(0,0,0,3/16.0f), FSlateColor::UseSubduedForeground()))
+			.SetHoveredBackgroundBrush(FSlateNoResource())
+			.SetInactiveFillBrush(FSlateNoResource())
+			.SetActiveFillBrush(FSlateNoResource())
+			.SetForegroundColor(FSlateColor::UseSubduedForeground())
+			.SetArrowsImage(FSlateNoResource())
+		);
 	}
 
 	// Foliage Edit Mode
diff --git a/Engine/Source/Editor/MovieSceneTools/Private/CurveKeyEditors/SFloatCurveKeyEditor.cpp b/Engine/Source/Editor/MovieSceneTools/Private/CurveKeyEditors/SFloatCurveKeyEditor.cpp
--- a/Engine/Source/Editor/MovieSceneTools/Private/CurveKeyEditors/SFloatCurveKeyEditor.cpp
+++ b/Engine/Source/Editor/MovieSceneTools/Private/CurveKeyEditors/SFloatCurveKeyEditor.cpp
@@ -13,19 +13,19 @@ void SFloatCurveKeyEditor::Construct(const FArguments& InArgs)
 	ChildSlot
 	[
 		SNew(SSpinBox<float>)
-		.Style(&FEditorStyle::GetWidgetStyle<FSpinBoxStyle>("Sequencer.AnimationOutliner.KeyEditorSpinBoxStyle"))
+		.Style(&FEditorStyle::GetWidgetStyle<FSpinBoxStyle>("Sequencer.HyperlinkSpinBox"))
 		.Font(FEditorStyle::GetFontStyle("Sequencer.AnimationOutliner.RegularFont"))
 		.MinValue(TOptional<float>())
 		.MaxValue(TOptional<float>())
 		.MaxSliderValue(TOptional<float>())
 		.MinSliderValue(TOptional<float>())
 		.Delta(0.001f)
-		.MinDesiredWidth(60.0f)
 		.Value(this, &SFloatCurveKeyEditor::OnGetKeyValue)
 		.OnValueChanged(this, &SFloatCurveKeyEditor::OnValueChanged)
 		.OnValueCommitted(this, &SFloatCurveKeyEditor::OnValueCommitted)
 		.OnBeginSliderMovement(this, &SFloatCurveKeyEditor::OnBeginSliderMovement)
 		.OnEndSliderMovement(this, &SFloatCurveKeyEditor::OnEndSliderMovement)
+		.ClearKeyboardFocusOnCommit(true)
 	];
 }
 
diff --git a/Engine/Source/Editor/MovieSceneTools/Private/CurveKeyEditors/SIntegralCurveKeyEditor.cpp b/Engine/Source/Editor/MovieSceneTools/Private/CurveKeyEditors/SIntegralCurveKeyEditor.cpp
--- a/Engine/Source/Editor/MovieSceneTools/Private/CurveKeyEditors/SIntegralCurveKeyEditor.cpp
+++ b/Engine/Source/Editor/MovieSceneTools/Private/CurveKeyEditors/SIntegralCurveKeyEditor.cpp
@@ -13,19 +13,19 @@ void SIntegralCurveKeyEditor::Construct(const FArguments& InArgs)
 	ChildSlot
 	[
 		SNew(SSpinBox<int32>)
-		.Style(&FEditorStyle::GetWidgetStyle<FSpinBoxStyle>("Sequencer.AnimationOutliner.KeyEditorSpinBoxStyle"))
+		.Style(&FEditorStyle::GetWidgetStyle<FSpinBoxStyle>("Sequencer.HyperlinkSpinBox"))
 		.Font(FEditorStyle::GetFontStyle("Sequencer.AnimationOutliner.RegularFont"))
 		.MinValue(TOptional<int32>())
 		.MaxValue(TOptional<int32>())
 		.MaxSliderValue(TOptional<int32>())
 		.MinSliderValue(TOptional<int32>())
 		.Delta(1)
-		.MinDesiredWidth(60.0f)
 		.Value(this, &SIntegralCurveKeyEditor::OnGetKeyValue)
 		.OnValueChanged(this, &SIntegralCurveKeyEditor::OnValueChanged)
 		.OnValueCommitted(this, &SIntegralCurveKeyEditor::OnValueCommitted)
 		.OnBeginSliderMovement(this, &SIntegralCurveKeyEditor::OnBeginSliderMovement)
 		.OnEndSliderMovement(this, &SIntegralCurveKeyEditor::OnEndSliderMovement)
+		.ClearKeyboardFocusOnCommit(true)
 	];
 }
 
diff --git a/Engine/Source/Editor/Sequencer/Private/SAnimationOutlinerTreeNode.cpp b/Engine/Source/Editor/Sequencer/Private/SAnimationOutlinerTreeNode.cpp
--- a/Engine/Source/Editor/Sequencer/Private/SAnimationOutlinerTreeNode.cpp
+++ b/Engine/Source/Editor/Sequencer/Private/SAnimationOutlinerTreeNode.cpp
@@ -74,6 +74,7 @@ void SAnimationOutlinerTreeNode::Construct( const FArguments& InArgs, TSharedRef
 				[
 					// @todo Sequencer - Remove this box and width override.
 					SNew(SBox)
+					.HAlign(HAlign_Left)
 					.WidthOverride(100)
 					[
 						DisplayNode->GenerateEditWidgetForOutliner()
diff --git a/Engine/Source/Editor/SequencerWidgets/Private/STimeRange.cpp b/Engine/Source/Editor/SequencerWidgets/Private/STimeRange.cpp
--- a/Engine/Source/Editor/SequencerWidgets/Private/STimeRange.cpp
+++ b/Engine/Source/Editor/SequencerWidgets/Private/STimeRange.cpp
@@ -7,23 +7,14 @@
 
 #define LOCTEXT_NAMESPACE "STimeRange"
 
-/*
-* Custom time range text box for the inputs
-*/
-class STimeRangeTextBox : public SEditableTextBox
-{
-public:
-	void Construct(const FArguments& InArgs) { SEditableTextBox::Construct(InArgs); }
-protected:
-	virtual FVector2D ComputeDesiredSize(float) const override { return FVector2D(46.f, 16.f); }
-};
-
 void STimeRange::Construct( const STimeRange::FArguments& InArgs, TSharedRef<ITimeSliderController> InTimeSliderController )
 {
 	TimeSliderController = InTimeSliderController;
 	ShowFrameNumbers = InArgs._ShowFrameNumbers;
 	TimeSnapInterval = InArgs._TimeSnapInterval;
-	
+
+	auto NumericTypeInterface = SharedThis(this);
+
 	this->ChildSlot
 	.HAlign(HAlign_Fill)
 	[
@@ -32,19 +23,41 @@ void STimeRange::Construct( const STimeRange::FArguments& InArgs, TSharedRef<ITi
 		.AutoWidth()
 		.Padding(2.0f, 2.0f)
 		[
-			SNew(STimeRangeTextBox)
-			.Text(this, &STimeRange::StartTime)
-			.ToolTipText(this, &STimeRange::StartTimeTooltip )
-			.OnTextCommitted( this, &STimeRange::OnStartTimeCommitted )
+			SNew(SBox)
+			.MinDesiredWidth(18)
+			.HAlign(HAlign_Center)
+			[
+				SNew(SSpinBox<float>)
+				.Value(this, &STimeRange::StartTime)
+				.ToolTipText(this, &STimeRange::StartTimeTooltip )
+				.OnValueCommitted( this, &STimeRange::OnStartTimeCommitted )
+				.OnValueChanged( this, &STimeRange::OnStartTimeChanged )
+				.MinValue(TOptional<float>())
+				.MaxValue(this, &STimeRange::MaxStartTime)
+				.Style(&FEditorStyle::Get().GetWidgetStyle<FSpinBoxStyle>("Sequencer.HyperlinkSpinBox"))
+				.TypeInterface(NumericTypeInterface)
+				.ClearKeyboardFocusOnCommit(true)
+			]
 		]
 		+SHorizontalBox::Slot()
 		.AutoWidth()
 		.Padding(2.0f, 2.0f)
 		[
-			SNew(STimeRangeTextBox)			
-			.Text(this, &STimeRange::InTime)
-			.ToolTipText(this, &STimeRange::InTimeTooltip)
-			.OnTextCommitted( this, &STimeRange::OnInTimeCommitted )
+			SNew(SBox)
+			.MinDesiredWidth(18)
+			.HAlign(HAlign_Center)
+			[
+				SNew(SSpinBox<float>)
+				.Value(this, &STimeRange::InTime)
+				.ToolTipText(this, &STimeRange::InTimeTooltip)
+				.OnValueCommitted( this, &STimeRange::OnInTimeCommitted )
+				.OnValueChanged( this, &STimeRange::OnInTimeChanged )
+				.MinValue(this, &STimeRange::MinInTime)
+				.MaxValue(this, &STimeRange::MaxInTime)
+				.Style(&FEditorStyle::Get().GetWidgetStyle<FSpinBoxStyle>("Sequencer.HyperlinkSpinBox"))
+				.TypeInterface(NumericTypeInterface)
+				.ClearKeyboardFocusOnCommit(true)
+			]
 		]
 		+SHorizontalBox::Slot()
 		.HAlign(HAlign_Fill)
@@ -57,219 +70,164 @@ void STimeRange::Construct( const STimeRange::FArguments& InArgs, TSharedRef<ITi
 		.AutoWidth()
 		.Padding(2.0f, 2.0f)
 		[
-			SNew(STimeRangeTextBox)
-			.Text(this, &STimeRange::OutTime)
-			.ToolTipText(this, &STimeRange::OutTimeTooltip)
-			.OnTextCommitted( this, &STimeRange::OnOutTimeCommitted )
+			SNew(SBox)
+			.MinDesiredWidth(18)
+			.HAlign(HAlign_Center)
+			[
+				SNew(SSpinBox<float>)
+				.Value(this, &STimeRange::OutTime)
+				.ToolTipText(this, &STimeRange::OutTimeTooltip)
+				.OnValueCommitted( this, &STimeRange::OnOutTimeCommitted )
+				.OnValueChanged( this, &STimeRange::OnOutTimeChanged )
+				.MinValue(this, &STimeRange::MinOutTime)
+				.MaxValue(this, &STimeRange::MaxOutTime)
+				.Style(&FEditorStyle::Get().GetWidgetStyle<FSpinBoxStyle>("Sequencer.HyperlinkSpinBox"))
+				.TypeInterface(NumericTypeInterface)
+				.ClearKeyboardFocusOnCommit(true)
+			]
 		]
 		+SHorizontalBox::Slot()
 		.AutoWidth()
 		.Padding(2.0f, 2.0f)
 		[
-			SNew(STimeRangeTextBox)
-			.Text(this, &STimeRange::EndTime)
-			.ToolTipText(this, &STimeRange::EndTimeTooltip)
-			.OnTextCommitted( this, &STimeRange::OnEndTimeCommitted )
+			SNew(SBox)
+			.MinDesiredWidth(18)
+			.HAlign(HAlign_Center)
+			[
+				SNew(SSpinBox<float>)
+				.Value(this, &STimeRange::EndTime)
+				.ToolTipText(this, &STimeRange::EndTimeTooltip)
+				.OnValueCommitted( this, &STimeRange::OnEndTimeCommitted )
+				.OnValueChanged( this, &STimeRange::OnEndTimeChanged )
+				.MinValue(this, &STimeRange::MinEndTime)
+				.MaxValue(TOptional<float>())
+				.Style(&FEditorStyle::Get().GetWidgetStyle<FSpinBoxStyle>("Sequencer.HyperlinkSpinBox"))
+				.TypeInterface(NumericTypeInterface)
+				.ClearKeyboardFocusOnCommit(true)
+			]
 		]
 	];
 }
 
-FText STimeRange::StartTime() const
+float STimeRange::StartTime() const
 {
-	if (TimeSliderController.IsValid())
-	{
-		float StartTime = TimeSliderController.Get()->GetClampRange().GetLowerBoundValue();
-		bool bShowFrameNumbers = ShowFrameNumbers.IsBound() ? ShowFrameNumbers.Get() : false;
-
-		FString StartTimeStr;
-		if (bShowFrameNumbers)
-		{
-			int32 Frame = TimeSliderController.Get()->TimeToFrame(StartTime);
-			StartTimeStr = FString::Printf( TEXT("%d"), Frame);
-		}
-		else
-		{
-			StartTimeStr = FString::Printf( TEXT("%.2f"), StartTime);
-		}
-
-		return FText::FromString(StartTimeStr);
-	}
-
-	return FText::FromString("");
+	return TimeSliderController.IsValid() ? TimeSliderController.Get()->GetClampRange().GetLowerBoundValue() : 0.f;
 }
 
-FText STimeRange::EndTime() const
+float STimeRange::EndTime() const
 {
-	if (TimeSliderController.IsValid())
-	{
-		float EndTime = TimeSliderController.Get()->GetClampRange().GetUpperBoundValue();
-		bool bShowFrameNumbers = ShowFrameNumbers.IsBound() ? ShowFrameNumbers.Get() : false;
+	return TimeSliderController.IsValid() ? TimeSliderController.Get()->GetClampRange().GetUpperBoundValue() : 0.f;
+}
 
-		FString EndTimeStr;
-		if (bShowFrameNumbers)
-		{
-			int32 Frame = TimeSliderController.Get()->TimeToFrame(EndTime);
-			EndTimeStr = FString::Printf( TEXT("%d"), Frame);
-		}
-		else
-		{
-			EndTimeStr = FString::Printf( TEXT("%.2f"), EndTime);
-		}
+float STimeRange::InTime() const
+{
+	return TimeSliderController.IsValid() ? TimeSliderController.Get()->GetViewRange().GetLowerBoundValue() : 0.f;
+}
 
-		return FText::FromString(EndTimeStr);
-	}
+float STimeRange::OutTime() const
+{
+	return TimeSliderController.IsValid() ? TimeSliderController.Get()->GetViewRange().GetUpperBoundValue() : 0.f;
+}
 
-	return FText::FromString("");
+TOptional<float> STimeRange::MinInTime() const
+{
+	return StartTime();
 }
 
-FText STimeRange::InTime() const
+TOptional<float> STimeRange::MaxInTime() const
 {
-	if (TimeSliderController.IsValid())
-	{
-		float InTime = TimeSliderController.Get()->GetViewRange().GetLowerBoundValue();
-		bool bShowFrameNumbers = ShowFrameNumbers.IsBound() ? ShowFrameNumbers.Get() : false;
+	return OutTime();
+}
 
-		FString InTimeStr;
-		if (bShowFrameNumbers)
-		{
-			int32 Frame = TimeSliderController.Get()->TimeToFrame(InTime);
-			InTimeStr = FString::Printf( TEXT("%d"), Frame);
-		}
-		else
-		{
-			InTimeStr = FString::Printf( TEXT("%.2f"), InTime);
-		}
+TOptional<float> STimeRange::MinOutTime() const
+{
+	return InTime();
+}
 
-		return FText::FromString(InTimeStr);
-	}
+TOptional<float> STimeRange::MaxOutTime() const
+{
+	return EndTime();
+}
 
-	return FText::FromString("");
+TOptional<float> STimeRange::MaxStartTime() const
+{
+	return OutTime();
 }
 
-FText STimeRange::OutTime() const
+TOptional<float> STimeRange::MinEndTime() const
 {
-	if (TimeSliderController.IsValid())
-	{
-		float OutTime = TimeSliderController.Get()->GetViewRange().GetUpperBoundValue();
-		bool bShowFrameNumbers = ShowFrameNumbers.IsBound() ? ShowFrameNumbers.Get() : false;
+	return InTime();
+}
 
-		FString OutTimeStr;
-		if (bShowFrameNumbers)
-		{
-			int32 Frame = TimeSliderController.Get()->TimeToFrame(OutTime);
-			OutTimeStr = FString::Printf( TEXT("%d"), Frame);
-		}
-		else
-		{
-			OutTimeStr = FString::Printf( TEXT("%.2f"), OutTime); 
-		}
+void STimeRange::OnStartTimeCommitted(float NewValue, ETextCommit::Type InTextCommit)
+{
+	OnStartTimeChanged(NewValue);
+}
 
-		return FText::FromString(OutTimeStr);
-	}
+void STimeRange::OnEndTimeCommitted(float NewValue, ETextCommit::Type InTextCommit)
+{
+	OnEndTimeChanged(NewValue);
+}
 
-	return FText::FromString("");
+void STimeRange::OnInTimeCommitted(float NewValue, ETextCommit::Type InTextCommit)
+{
+	OnInTimeChanged(NewValue);
 }
 
-void STimeRange::OnStartTimeCommitted(const FText& NewText, ETextCommit::Type InTextCommit)
+void STimeRange::OnOutTimeCommitted(float NewValue, ETextCommit::Type InTextCommit)
 {
-	FString NewTextStr = NewText.ToString();
+	OnOutTimeChanged(NewValue);
+}
 
-	if (TimeSliderController.IsValid())
+void STimeRange::OnStartTimeChanged(float NewValue)
+{
+	if (auto* Controller = TimeSliderController.Get())
 	{
-		float NewStart;
-		bool bShowFrameNumbers = ShowFrameNumbers.IsBound() ? ShowFrameNumbers.Get() : false;
+		Controller->SetClampRange(NewValue, Controller->GetClampRange().GetUpperBoundValue());
 
-		if (bShowFrameNumbers)
+		if (NewValue > Controller->GetViewRange().GetLowerBoundValue())
 		{
-			int32 NewStartFrame = FCString::Atoi(*NewTextStr);
-			NewStart = TimeSliderController.Get()->FrameToTime(NewStartFrame);
+			Controller->SetViewRange(NewValue, Controller->GetViewRange().GetUpperBoundValue(), EViewRangeInterpolation::Immediate);
 		}
-		else
-		{
-			NewStart = FCString::Atof(*NewTextStr);
-		}
-
-		TimeSliderController.Get()->SetClampRange(NewStart, TimeSliderController.Get()->GetClampRange().GetUpperBoundValue());
 	}
 }
 
-void STimeRange::OnEndTimeCommitted(const FText& NewText, ETextCommit::Type InTextCommit)
+void STimeRange::OnEndTimeChanged(float NewValue)
 {
-	FString NewTextStr = NewText.ToString();
-
-	if (TimeSliderController.IsValid())
+	if (auto* Controller = TimeSliderController.Get())
 	{
-		float NewEnd;
-		bool bShowFrameNumbers = ShowFrameNumbers.IsBound() ? ShowFrameNumbers.Get() : false;
+		Controller->SetClampRange(Controller->GetClampRange().GetLowerBoundValue(), NewValue);
 
-		if (bShowFrameNumbers)
-		{
-			int32 NewEndFrame = FCString::Atoi(*NewTextStr);
-			NewEnd = TimeSliderController.Get()->FrameToTime(NewEndFrame);
-		}
-		else
+		if (NewValue < Controller->GetViewRange().GetUpperBoundValue())
 		{
-			NewEnd = FCString::Atof(*NewTextStr);
+			Controller->SetViewRange(Controller->GetViewRange().GetLowerBoundValue(), NewValue, EViewRangeInterpolation::Immediate);
 		}
-
-		TimeSliderController.Get()->SetClampRange(TimeSliderController.Get()->GetClampRange().GetLowerBoundValue(), NewEnd);
 	}
 }
 
-void STimeRange::OnInTimeCommitted(const FText& NewText, ETextCommit::Type InTextCommit)
+void STimeRange::OnInTimeChanged(float NewValue)
 {
-	FString NewTextStr = NewText.ToString();
-
-	if (TimeSliderController.IsValid())
+	if (auto* Controller = TimeSliderController.Get())
 	{
-		float NewIn;
-		bool bShowFrameNumbers = ShowFrameNumbers.IsBound() ? ShowFrameNumbers.Get() : false;
-
-		if (bShowFrameNumbers)
-		{
-			int32 NewInFrame = FCString::Atoi(*NewTextStr);
-
-			NewIn = TimeSliderController.Get()->FrameToTime(NewInFrame);
-		}
-		else
+		if (NewValue < TimeSliderController.Get()->GetClampRange().GetLowerBoundValue())
 		{
-			NewIn = FCString::Atof(*NewTextStr);
+			Controller->SetClampRange(NewValue, Controller->GetClampRange().GetUpperBoundValue());
 		}
 
-		if (NewIn < TimeSliderController.Get()->GetClampRange().GetLowerBoundValue())
-		{
-			TimeSliderController.Get()->SetClampRange(NewIn, TimeSliderController.Get()->GetClampRange().GetUpperBoundValue());
-		}
-
-		TimeSliderController.Get()->SetViewRange(NewIn, TimeSliderController.Get()->GetViewRange().GetUpperBoundValue(), EViewRangeInterpolation::Immediate);
+		Controller->SetViewRange(NewValue, Controller->GetViewRange().GetUpperBoundValue(), EViewRangeInterpolation::Immediate);
 	}
 }
 
-void STimeRange::OnOutTimeCommitted(const FText& NewText, ETextCommit::Type InTextCommit)
+void STimeRange::OnOutTimeChanged(float NewValue)
 {
-	FString NewTextStr = NewText.ToString();
-
-	if (TimeSliderController.IsValid())
+	if (auto* Controller = TimeSliderController.Get())
 	{
-		float NewOut;
-		bool bShowFrameNumbers = ShowFrameNumbers.IsBound() ? ShowFrameNumbers.Get() : false;
-
-		if (bShowFrameNumbers)
+		if (NewValue > Controller->GetClampRange().GetUpperBoundValue())
 		{
-			int32 NewOutFrame = FCString::Atoi(*NewTextStr);
-			NewOut = TimeSliderController.Get()->FrameToTime(NewOutFrame);
-		}
-		else
-		{
-			NewOut = FCString::Atof(*NewTextStr);
-		}
-
-		if (NewOut > TimeSliderController.Get()->GetClampRange().GetUpperBoundValue())
-		{
-			TimeSliderController.Get()->SetClampRange(TimeSliderController.Get()->GetClampRange().GetLowerBoundValue(), NewOut);
+			Controller->SetClampRange(Controller->GetClampRange().GetLowerBoundValue(), NewValue);
 		}
 
-		TimeSliderController.Get()->SetViewRange(TimeSliderController.Get()->GetViewRange().GetLowerBoundValue(), NewOut, EViewRangeInterpolation::Immediate);
+		Controller->SetViewRange(Controller->GetViewRange().GetLowerBoundValue(), NewValue, EViewRangeInterpolation::Immediate);
 	}
 }
 
@@ -334,4 +292,31 @@ float STimeRange::GetTimeSnapInterval() const
 	return 1.f;
 }
 
+FString STimeRange::ToString(const float& Value) const
+{
+	bool bShowFrameNumbers = ShowFrameNumbers.IsBound() ? ShowFrameNumbers.Get() : false;
+
+	if (bShowFrameNumbers && TimeSliderController.IsValid())
+	{
+		int32 Frame = TimeSliderController.Get()->TimeToFrame(Value);
+		return FString::Printf(TEXT("%d"), Frame);
+	}
+
+	return FString::Printf(TEXT("%.2f"), Value);
+}
+
+TOptional<float> STimeRange::FromString(const FString& InString)
+{
+	bool bShowFrameNumbers = ShowFrameNumbers.IsBound() ? ShowFrameNumbers.Get() : false;
+	if (bShowFrameNumbers && TimeSliderController.IsValid())
+	{
+		int32 NewEndFrame = FCString::Atoi(*InString);
+		return float(TimeSliderController.Get()->FrameToTime(NewEndFrame));
+	}
+
+	return FCString::Atof(*InString);
+}
+
+
+
 #undef LOCTEXT_NAMESPACE
diff --git a/Engine/Source/Editor/SequencerWidgets/Private/STimeRange.h b/Engine/Source/Editor/SequencerWidgets/Private/STimeRange.h
--- a/Engine/Source/Editor/SequencerWidgets/Private/STimeRange.h
+++ b/Engine/Source/Editor/SequencerWidgets/Private/STimeRange.h
@@ -8,7 +8,7 @@ class STimeRangeSlider;
 
 DECLARE_DELEGATE_RetVal( bool, STimeRangeGetter );
 
-class STimeRange : public ITimeSlider
+class STimeRange : public ITimeSlider, public TDefaultNumericTypeInterface<float>
 {
 public:
 	SLATE_BEGIN_ARGS(STimeRange)
@@ -29,21 +29,37 @@ public:
 	float GetTimeSnapInterval() const;
 
 protected:
-	FText InTime() const;
-	FText OutTime() const;
-	FText StartTime() const;
-	FText EndTime() const;
+	float InTime() const;
+	float OutTime() const;
+	float StartTime() const;
+	float EndTime() const;
 
-	void OnStartTimeCommitted(const FText& NewText, ETextCommit::Type InTextCommit);
-	void OnEndTimeCommitted(const FText& NewText, ETextCommit::Type InTextCommit);
-	void OnInTimeCommitted(const FText& NewText, ETextCommit::Type InTextCommit);
-	void OnOutTimeCommitted(const FText& NewText, ETextCommit::Type InTextCommit);
+	TOptional<float> MinInTime() const;
+	TOptional<float> MaxInTime() const;
+	TOptional<float> MinOutTime() const;
+	TOptional<float> MaxOutTime() const;
+	TOptional<float> MaxStartTime() const;
+	TOptional<float> MinEndTime() const;
+	
+	void OnStartTimeCommitted(float NewValue, ETextCommit::Type InTextCommit);
+	void OnEndTimeCommitted(float NewValue, ETextCommit::Type InTextCommit);
+	void OnInTimeCommitted(float NewValue, ETextCommit::Type InTextCommit);
+	void OnOutTimeCommitted(float NewValue, ETextCommit::Type InTextCommit);
+
+	void OnStartTimeChanged(float NewValue);
+	void OnEndTimeChanged(float NewValue);
+	void OnInTimeChanged(float NewValue);
+	void OnOutTimeChanged(float NewValue);
 
 	FText InTimeTooltip() const;
 	FText OutTimeTooltip() const;
 	FText StartTimeTooltip() const;
 	FText EndTimeTooltip() const;
 
+	/** Convert the type to/from a string */
+	virtual FString ToString(const float& Value) const override;
+	virtual TOptional<float> FromString(const FString& InString) override;
+
 private:
 	TSharedPtr<ITimeSliderController> TimeSliderController;
 	TSharedPtr<STimeRangeSlider> TimeRangeSlider;
