Added missing header

diff --git a/Engine/Source/Editor/SequencerWidgets/Private/STimeRange.cpp b/Engine/Source/Editor/SequencerWidgets/Private/STimeRange.cpp
--- a/Engine/Source/Editor/SequencerWidgets/Private/STimeRange.cpp
+++ b/Engine/Source/Editor/SequencerWidgets/Private/STimeRange.cpp
@@ -4,6 +4,7 @@
 #include "STimeRange.h"
 #include "STimeRangeSlider.h"
 #include "SlateStyle.h"
+#include "EditorStyle.h"
 
 #define LOCTEXT_NAMESPACE "STimeRange"
 
