Fixing hidden virtual function on base class

diff --git a/Engine/Source/Editor/SequencerWidgets/Private/STimeRange.h b/Engine/Source/Editor/SequencerWidgets/Private/STimeRange.h
--- a/Engine/Source/Editor/SequencerWidgets/Private/STimeRange.h
+++ b/Engine/Source/Editor/SequencerWidgets/Private/STimeRange.h
@@ -28,6 +28,8 @@ public:
 
 	float GetTimeSnapInterval() const;
 
+	using ITimeSlider::ToString;
+	
 protected:
 	float InTime() const;
 	float OutTime() const;
