Copying //UE4/Dev-Sequencer to //UE4/Main

diff --git a/Engine/Source/Editor/SequencerWidgets/Private/STimeRange.cpp b/Engine/Source/Editor/SequencerWidgets/Private/STimeRange.cpp
--- a/Engine/Source/Editor/SequencerWidgets/Private/STimeRange.cpp
+++ b/Engine/Source/Editor/SequencerWidgets/Private/STimeRange.cpp
@@ -8,179 +8,299 @@
 
 #define LOCTEXT_NAMESPACE "STimeRange"
 
-void STimeRange::Construct( const STimeRange::FArguments& InArgs, TSharedRef<ITimeSliderController> InTimeSliderController )
+void STimeRange::Construct( const STimeRange::FArguments& InArgs, TSharedRef<ITimeSliderController> InTimeSliderController, TSharedRef<INumericTypeInterface<float>> NumericTypeInterface )
 {
 	TimeSliderController = InTimeSliderController;
 	ShowFrameNumbers = InArgs._ShowFrameNumbers;
-	TimeSnapInterval = InArgs._TimeSnapInterval;
 
-	auto NumericTypeInterface = SharedThis(this);
+	TSharedRef<SWidget> WorkingRangeStart = SNullWidget::NullWidget, WorkingRangeEnd = SNullWidget::NullWidget;
+	if (InArgs._ShowWorkingRange)
+	{
+		WorkingRangeStart = SNew(SSpinBox<float>)
+		.Value(this, &STimeRange::WorkingStartTime)
+		.ToolTipText(LOCTEXT("WorkingRangeStart", "Working Range Start"))
+		.OnValueCommitted( this, &STimeRange::OnWorkingStartTimeCommitted )
+		.OnValueChanged( this, &STimeRange::OnWorkingStartTimeChanged )
+		.MinValue(TOptional<float>())
+		.MaxValue(this, &STimeRange::MaxWorkingStartTime)
+		.Style(&FEditorStyle::Get().GetWidgetStyle<FSpinBoxStyle>("Sequencer.HyperlinkSpinBox"))
+		.TypeInterface(NumericTypeInterface)
+		.ClearKeyboardFocusOnCommit(true);
+
+		WorkingRangeEnd = SNew(SSpinBox<float>)
+		.Value(this, &STimeRange::WorkingEndTime)
+		.ToolTipText(LOCTEXT("WorkingRangeEnd", "Working Range End"))
+		.OnValueCommitted( this, &STimeRange::OnWorkingEndTimeCommitted )
+		.OnValueChanged( this, &STimeRange::OnWorkingEndTimeChanged )
+		.MinValue(this, &STimeRange::MinWorkingEndTime)
+		.MaxValue(TOptional<float>())
+		.Style(&FEditorStyle::Get().GetWidgetStyle<FSpinBoxStyle>("Sequencer.HyperlinkSpinBox"))
+		.TypeInterface(NumericTypeInterface)
+		.ClearKeyboardFocusOnCommit(true);
+	}
+
+	TSharedRef<SWidget> ViewRangeStart = SNullWidget::NullWidget, ViewRangeEnd = SNullWidget::NullWidget;
+	if (InArgs._ShowViewRange)
+	{
+		ViewRangeStart = SNew(SSpinBox<float>)
+		.Value(this, &STimeRange::ViewStartTime)
+		.ToolTipText(ViewStartTimeTooltip())
+		.OnValueCommitted( this, &STimeRange::OnViewStartTimeCommitted )
+		.OnValueChanged( this, &STimeRange::OnViewStartTimeChanged )
+		.MinValue(TOptional<float>())
+		.MaxValue(this, &STimeRange::MaxViewStartTime)
+		.Style(&FEditorStyle::Get().GetWidgetStyle<FSpinBoxStyle>("Sequencer.HyperlinkSpinBox"))
+		.TypeInterface(NumericTypeInterface)
+		.ClearKeyboardFocusOnCommit(true);
+
+		ViewRangeEnd = SNew(SSpinBox<float>)
+		.Value(this, &STimeRange::ViewEndTime)
+		.ToolTipText(ViewEndTimeTooltip())
+		.OnValueCommitted( this, &STimeRange::OnViewEndTimeCommitted )
+		.OnValueChanged( this, &STimeRange::OnViewEndTimeChanged )
+		.MinValue(this, &STimeRange::MinViewEndTime)
+		.MaxValue(TOptional<float>())
+		.Style(&FEditorStyle::Get().GetWidgetStyle<FSpinBoxStyle>("Sequencer.HyperlinkSpinBox"))
+		.TypeInterface(NumericTypeInterface)
+		.ClearKeyboardFocusOnCommit(true);
+	}
+
+	TSharedRef<SWidget> PlaybackRangeStart = SNullWidget::NullWidget, PlaybackRangeEnd = SNullWidget::NullWidget;
+	{
+		PlaybackRangeStart = SNew(SSpinBox<float>)
+		.Value(this, &STimeRange::PlayStartTime)
+		.ToolTipText(PlayStartTimeTooltip())
+		.OnValueCommitted( this, &STimeRange::OnPlayStartTimeCommitted )
+		.OnValueChanged( this, &STimeRange::OnPlayStartTimeChanged )
+		.MinValue(this, &STimeRange::MinPlayStartTime)
+		.MaxValue(this, &STimeRange::MaxPlayStartTime)
+		.Style(&FEditorStyle::Get().GetWidgetStyle<FSpinBoxStyle>("Sequencer.HyperlinkSpinBox"))
+		.TypeInterface(NumericTypeInterface)
+		.ClearKeyboardFocusOnCommit(true);
+
+		PlaybackRangeEnd = SNew(SSpinBox<float>)
+		.Value(this, &STimeRange::PlayEndTime)
+		.ToolTipText(PlayEndTimeTooltip())
+		.OnValueCommitted( this, &STimeRange::OnPlayEndTimeCommitted )
+		.OnValueChanged( this, &STimeRange::OnPlayEndTimeChanged )
+		.MinValue(this, &STimeRange::MinPlayEndTime)
+		.MaxValue(this, &STimeRange::MaxPlayEndTime)
+		.Style(&FEditorStyle::Get().GetWidgetStyle<FSpinBoxStyle>("Sequencer.HyperlinkSpinBox"))
+		.TypeInterface(NumericTypeInterface)
+		.ClearKeyboardFocusOnCommit(true);
+	}
 
 	this->ChildSlot
 	.HAlign(HAlign_Fill)
 	[
 		SNew(SHorizontalBox)
+
 		+SHorizontalBox::Slot()
+		.VAlign(VAlign_Center)
 		.AutoWidth()
-		.Padding(2.0f, 2.0f)
+		.Padding(2.f)
 		[
 			SNew(SBox)
+			.Visibility(InArgs._ShowWorkingRange ? EVisibility::Visible : EVisibility::Collapsed)
 			.WidthOverride(36)
 			.HAlign(HAlign_Center)
 			[
-				SNew(SSpinBox<float>)
-				.Value(this, &STimeRange::StartTime)
-				.ToolTipText(this, &STimeRange::StartTimeTooltip )
-				.OnValueCommitted( this, &STimeRange::OnStartTimeCommitted )
-				.OnValueChanged( this, &STimeRange::OnStartTimeChanged )
-				.MinValue(TOptional<float>())
-				.MaxValue(this, &STimeRange::MaxOutTime)
-				.Style(&FEditorStyle::Get().GetWidgetStyle<FSpinBoxStyle>("Sequencer.HyperlinkSpinBox"))
-				.TypeInterface(NumericTypeInterface)
-				.ClearKeyboardFocusOnCommit(true)
+				WorkingRangeStart
 			]
 		]
+
 		+SHorizontalBox::Slot()
+		.VAlign(VAlign_Center)
 		.AutoWidth()
-		.Padding(2.0f, 2.0f)
+		.Padding(2.f)
 		[
 			SNew(SBox)
 			.WidthOverride(36)
 			.HAlign(HAlign_Center)
+			.Visibility(InArgs._ShowPlaybackRange ? EVisibility::Visible : EVisibility::Collapsed)
 			[
-				SNew(SSpinBox<float>)
-				.Value(this, &STimeRange::InTime)
-				.ToolTipText(this, &STimeRange::InTimeTooltip)
-				.OnValueCommitted( this, &STimeRange::OnInTimeCommitted )
-				.OnValueChanged( this, &STimeRange::OnInTimeChanged )
-				.MinValue(TOptional<float>())
-				.MaxValue(this, &STimeRange::MaxOutTime)
-				.Style(&FEditorStyle::Get().GetWidgetStyle<FSpinBoxStyle>("Sequencer.HyperlinkSpinBox"))
-				.TypeInterface(NumericTypeInterface)
-				.ClearKeyboardFocusOnCommit(true)
+				SNew(SBorder)
+				.Padding(0)
+				.BorderImage(nullptr)
+				.ForegroundColor(FLinearColor::Green)
+				[
+					PlaybackRangeStart
+				]
 			]
 		]
+
 		+SHorizontalBox::Slot()
-		.HAlign(HAlign_Fill)
-		.VAlign(VAlign_Fill)
+		.VAlign(VAlign_Center)
+		.AutoWidth()
+		.Padding(2.f)
+		[
+			SNew(SBox)
+			.Visibility(InArgs._ShowViewRange ? EVisibility::Visible : EVisibility::Collapsed)
+			.WidthOverride(36)
+			.HAlign(HAlign_Center)
+			[
+				ViewRangeStart
+			]
+		]
+
+		+SHorizontalBox::Slot()
+
 		.Padding(2.0f, 4.0f)
 		[
-			SAssignNew(TimeRangeSlider, STimeRangeSlider, InTimeSliderController, SharedThis(this))
+			InArgs._CenterContent.Widget
 		]
+		
 		+SHorizontalBox::Slot()
+		.VAlign(VAlign_Center)
 		.AutoWidth()
-		.Padding(2.0f, 2.0f)
+		.Padding(2.f)
 		[
 			SNew(SBox)
+			.Visibility(InArgs._ShowViewRange ? EVisibility::Visible : EVisibility::Collapsed)
 			.WidthOverride(36)
 			.HAlign(HAlign_Center)
 			[
-				SNew(SSpinBox<float>)
-				.Value(this, &STimeRange::OutTime)
-				.ToolTipText(this, &STimeRange::OutTimeTooltip)
-				.OnValueCommitted( this, &STimeRange::OnOutTimeCommitted )
-				.OnValueChanged( this, &STimeRange::OnOutTimeChanged )
-				.MinValue(this, &STimeRange::MinInTime)
-				.MaxValue(TOptional<float>())
-				.Style(&FEditorStyle::Get().GetWidgetStyle<FSpinBoxStyle>("Sequencer.HyperlinkSpinBox"))
-				.TypeInterface(NumericTypeInterface)
-				.ClearKeyboardFocusOnCommit(true)
+				ViewRangeEnd
 			]
 		]
+
+		+SHorizontalBox::Slot()
+		.VAlign(VAlign_Center)
+		.AutoWidth()
+		.Padding(2.f)
+		[
+			SNew(SBox)
+			.WidthOverride(36)
+			.HAlign(HAlign_Center)
+			.Visibility(InArgs._ShowPlaybackRange ? EVisibility::Visible : EVisibility::Collapsed)
+			[
+				SNew(SBorder)
+				.Padding(0)
+				.BorderImage(nullptr)
+				.ForegroundColor(FLinearColor::Red)
+				[
+					PlaybackRangeEnd
+				]
+			]
+		]
+
 		+SHorizontalBox::Slot()
+		.VAlign(VAlign_Center)
 		.AutoWidth()
-		.Padding(2.0f, 2.0f)
+		.Padding(2.f)
 		[
 			SNew(SBox)
 			.WidthOverride(36)
 			.HAlign(HAlign_Center)
+			.Visibility(InArgs._ShowWorkingRange ? EVisibility::Visible : EVisibility::Collapsed)
 			[
-				SNew(SSpinBox<float>)
-				.Value(this, &STimeRange::EndTime)
-				.ToolTipText(this, &STimeRange::EndTimeTooltip)
-				.OnValueCommitted( this, &STimeRange::OnEndTimeCommitted )
-				.OnValueChanged( this, &STimeRange::OnEndTimeChanged )
-				.MinValue(this, &STimeRange::MinInTime)
-				.MaxValue(TOptional<float>())
-				.Style(&FEditorStyle::Get().GetWidgetStyle<FSpinBoxStyle>("Sequencer.HyperlinkSpinBox"))
-				.TypeInterface(NumericTypeInterface)
-				.ClearKeyboardFocusOnCommit(true)
+				WorkingRangeEnd
 			]
 		]
 	];
 }
 
-float STimeRange::StartTime() const
+float STimeRange::WorkingStartTime() const
 {
 	return TimeSliderController.IsValid() ? TimeSliderController.Get()->GetClampRange().GetLowerBoundValue() : 0.f;
 }
 
-float STimeRange::EndTime() const
+float STimeRange::WorkingEndTime() const
 {
 	return TimeSliderController.IsValid() ? TimeSliderController.Get()->GetClampRange().GetUpperBoundValue() : 0.f;
 }
 
-float STimeRange::InTime() const
+float STimeRange::ViewStartTime() const
 {
 	return TimeSliderController.IsValid() ? TimeSliderController.Get()->GetViewRange().GetLowerBoundValue() : 0.f;
 }
 
-float STimeRange::OutTime() const
+float STimeRange::ViewEndTime() const
 {
 	return TimeSliderController.IsValid() ? TimeSliderController.Get()->GetViewRange().GetUpperBoundValue() : 0.f;
 }
 
-TOptional<float> STimeRange::MinInTime() const
+float STimeRange::PlayStartTime() const
 {
-	return StartTime();
+	return TimeSliderController.IsValid() ? TimeSliderController.Get()->GetPlayRange().GetLowerBoundValue() : 0.f;
 }
 
-TOptional<float> STimeRange::MaxInTime() const
+float STimeRange::PlayEndTime() const
 {
-	return OutTime();
+	return TimeSliderController.IsValid() ? TimeSliderController.Get()->GetPlayRange().GetUpperBoundValue() : 0.f;
 }
 
-TOptional<float> STimeRange::MinOutTime() const
+TOptional<float> STimeRange::MaxViewStartTime() const
 {
-	return InTime();
+	return ViewEndTime();
 }
 
-TOptional<float> STimeRange::MaxOutTime() const
+TOptional<float> STimeRange::MinViewEndTime() const
 {
-	return EndTime();
+	return ViewStartTime();
 }
 
-TOptional<float> STimeRange::MaxStartTime() const
+TOptional<float> STimeRange::MinPlayStartTime() const
 {
-	return OutTime();
+	return WorkingStartTime();
 }
 
-TOptional<float> STimeRange::MinEndTime() const
+TOptional<float> STimeRange::MaxPlayStartTime() const
 {
-	return InTime();
+	return PlayEndTime();
 }
 
-void STimeRange::OnStartTimeCommitted(float NewValue, ETextCommit::Type InTextCommit)
+TOptional<float> STimeRange::MinPlayEndTime() const
 {
-	OnStartTimeChanged(NewValue);
+	return PlayStartTime();
 }
 
-void STimeRange::OnEndTimeCommitted(float NewValue, ETextCommit::Type InTextCommit)
+TOptional<float> STimeRange::MaxPlayEndTime() const
 {
-	OnEndTimeChanged(NewValue);
+	return WorkingEndTime();
 }
 
-void STimeRange::OnInTimeCommitted(float NewValue, ETextCommit::Type InTextCommit)
+TOptional<float> STimeRange::MaxWorkingStartTime() const
 {
-	OnInTimeChanged(NewValue);
+	return ViewEndTime();
 }
 
-void STimeRange::OnOutTimeCommitted(float NewValue, ETextCommit::Type InTextCommit)
+TOptional<float> STimeRange::MinWorkingEndTime() const
 {
-	OnOutTimeChanged(NewValue);
+	return ViewStartTime();
 }
 
-void STimeRange::OnStartTimeChanged(float NewValue)
+void STimeRange::OnWorkingStartTimeCommitted(float NewValue, ETextCommit::Type InTextCommit)
+{
+	OnWorkingStartTimeChanged(NewValue);
+}
+
+void STimeRange::OnWorkingEndTimeCommitted(float NewValue, ETextCommit::Type InTextCommit)
+{
+	OnWorkingEndTimeChanged(NewValue);
+}
+
+void STimeRange::OnViewStartTimeCommitted(float NewValue, ETextCommit::Type InTextCommit)
+{
+	OnViewStartTimeChanged(NewValue);
+}
+
+void STimeRange::OnViewEndTimeCommitted(float NewValue, ETextCommit::Type InTextCommit)
+{
+	OnViewEndTimeChanged(NewValue);
+}
+
+void STimeRange::OnPlayStartTimeCommitted(float NewValue, ETextCommit::Type InTextCommit)
+{
+	OnPlayStartTimeChanged(NewValue);
+}
+
+void STimeRange::OnPlayEndTimeCommitted(float NewValue, ETextCommit::Type InTextCommit)
+{
+	OnPlayEndTimeChanged(NewValue);
+}
+
+void STimeRange::OnWorkingStartTimeChanged(float NewValue)
 {
 	if (auto* Controller = TimeSliderController.Get())
 	{
@@ -193,7 +313,7 @@ void STimeRange::OnStartTimeChanged(float NewValue)
 	}
 }
 
-void STimeRange::OnEndTimeChanged(float NewValue)
+void STimeRange::OnWorkingEndTimeChanged(float NewValue)
 {
 	if (auto* Controller = TimeSliderController.Get())
 	{
@@ -206,7 +326,7 @@ void STimeRange::OnEndTimeChanged(float NewValue)
 	}
 }
 
-void STimeRange::OnInTimeChanged(float NewValue)
+void STimeRange::OnViewStartTimeChanged(float NewValue)
 {
 	if (auto* Controller = TimeSliderController.Get())
 	{
@@ -219,7 +339,7 @@ void STimeRange::OnInTimeChanged(float NewValue)
 	}
 }
 
-void STimeRange::OnOutTimeChanged(float NewValue)
+void STimeRange::OnViewEndTimeChanged(float NewValue)
 {
 	if (auto* Controller = TimeSliderController.Get())
 	{
@@ -232,99 +352,82 @@ void STimeRange::OnOutTimeChanged(float NewValue)
 	}
 }
 
-FText STimeRange::StartTimeTooltip() const
+void STimeRange::OnPlayStartTimeChanged(float NewValue)
 {
-	bool bShowFrameNumbers = ShowFrameNumbers.IsBound() ? ShowFrameNumbers.Get() : false;
-	if (bShowFrameNumbers)
+	if (auto* Controller = TimeSliderController.Get())
 	{
-		return LOCTEXT("StartFrameTooltip", "Start frame");
+		if (NewValue < TimeSliderController.Get()->GetClampRange().GetLowerBoundValue())
+		{
+			Controller->SetClampRange(NewValue, Controller->GetClampRange().GetUpperBoundValue());
+		}
+
+		Controller->SetPlayRange(NewValue, Controller->GetPlayRange().GetUpperBoundValue());
 	}
-	else
+}
+
+void STimeRange::OnPlayEndTimeChanged(float NewValue)
+{
+	if (auto* Controller = TimeSliderController.Get())
 	{
-		return LOCTEXT("StartTimeTooltip", "Start time");
+		if (NewValue > Controller->GetClampRange().GetUpperBoundValue())
+		{
+			Controller->SetClampRange(Controller->GetClampRange().GetLowerBoundValue(), NewValue);
+		}
+
+		Controller->SetPlayRange(Controller->GetPlayRange().GetLowerBoundValue(), NewValue);
 	}
 }
 
-FText STimeRange::InTimeTooltip() const
+FText STimeRange::PlayStartTimeTooltip() const
 {
 	bool bShowFrameNumbers = ShowFrameNumbers.IsBound() ? ShowFrameNumbers.Get() : false;
 	if (bShowFrameNumbers)
 	{
-		return LOCTEXT("InFrameTooltip", "In frame");
+		return LOCTEXT("PlayStartFrameTooltip", "In Frame");
 	}
 	else
 	{
-		return LOCTEXT("InTimeTooltip", "In time");
+		return LOCTEXT("PlayStartTimeTooltip", "In Time");
 	}
 }
 
-FText STimeRange::OutTimeTooltip() const
+FText STimeRange::PlayEndTimeTooltip() const
 {
 	bool bShowFrameNumbers = ShowFrameNumbers.IsBound() ? ShowFrameNumbers.Get() : false;
 	if (bShowFrameNumbers)
 	{
-		return LOCTEXT("OutFrameTooltip", "Out frame");
+		return LOCTEXT("PlayEndFrameTooltip", "Out Frame");
 	}
 	else
 	{
-		return LOCTEXT("OutTimeTooltip", "Out time");
+		return LOCTEXT("PlayEndTimeTooltip", "Out Time");
 	}
 }
 
-FText STimeRange::EndTimeTooltip() const
+FText STimeRange::ViewStartTimeTooltip() const
 {
 	bool bShowFrameNumbers = ShowFrameNumbers.IsBound() ? ShowFrameNumbers.Get() : false;
 	if (bShowFrameNumbers)
 	{
-		return LOCTEXT("EndFrameTooltip", "End frame");
+		return LOCTEXT("ViewStartFrameTooltip", "View Range Start Frame");
 	}
 	else
 	{
-		return LOCTEXT("EndTimeTooltip", "End time");
+		return LOCTEXT("ViewStartTimeTooltip", "View Range Start Time");
 	}
 }
 
-float STimeRange::GetTimeSnapInterval() const
-{
-	if (TimeSnapInterval.IsBound())
-	{
-		return TimeSnapInterval.Get();
-	}
-	return 1.f;
-}
-
-FString STimeRange::ToString(const float& Value) const
+FText STimeRange::ViewEndTimeTooltip() const
 {
 	bool bShowFrameNumbers = ShowFrameNumbers.IsBound() ? ShowFrameNumbers.Get() : false;
-
-	if (bShowFrameNumbers && TimeSliderController.IsValid())
+	if (bShowFrameNumbers)
 	{
-		int32 Frame = TimeSliderController.Get()->TimeToFrame(Value);
-		return FString::Printf(TEXT("%d"), Frame);
+		return LOCTEXT("ViewEndFrameTooltip", "View Range End Frame");
 	}
-
-	return FString::Printf(TEXT("%.2f"), Value);
-}
-
-TOptional<float> STimeRange::FromString(const FString& InString, const float& ExistingValue)
-{
-	bool bShowFrameNumbers = ShowFrameNumbers.IsBound() ? ShowFrameNumbers.Get() : false;
-	if (bShowFrameNumbers && TimeSliderController.IsValid())
+	else
 	{
-		// Convert existing value to frames
-		float ExistingValueInFrames = TimeSliderController.Get()->TimeToFrame(ExistingValue);
-		TOptional<float> Result = TDefaultNumericTypeInterface<float>::FromString(InString, ExistingValueInFrames);
-
-		if (Result.IsSet())
-		{
-			int32 NewEndFrame = FMath::RoundToInt(Result.GetValue());
-			return float(TimeSliderController.Get()->FrameToTime(NewEndFrame));
-		}
+		return LOCTEXT("ViewEndTimeTooltip", "View Range End Time");
 	}
-		
-	return TDefaultNumericTypeInterface<float>::FromString(InString, ExistingValue);
 }
 
-
-
 #undef LOCTEXT_NAMESPACE
diff --git a/Engine/Source/Editor/SequencerWidgets/Private/STimeRange.h b/Engine/Source/Editor/SequencerWidgets/Private/STimeRange.h
--- a/Engine/Source/Editor/SequencerWidgets/Private/STimeRange.h
+++ b/Engine/Source/Editor/SequencerWidgets/Private/STimeRange.h
@@ -4,19 +4,22 @@
 
 #include "ITimeSlider.h"
 
-class STimeRangeSlider;
-
-DECLARE_DELEGATE_RetVal( bool, STimeRangeGetter );
-
-class STimeRange : public ITimeSlider, public TDefaultNumericTypeInterface<float>
+class STimeRange : public ITimeSlider
 {
 public:
 	SLATE_BEGIN_ARGS(STimeRange)
+		: _ShowWorkingRange(true), _ShowViewRange(false), _ShowPlaybackRange(false)
 	{}
 		/* If we should show frame numbers on the timeline */
 		SLATE_ARGUMENT( TAttribute<bool>, ShowFrameNumbers )
-		/* The time snap interval for the timeline */
-		SLATE_ARGUMENT( TAttribute<float>, TimeSnapInterval )
+		/** Whether to show the working range */
+		SLATE_ARGUMENT( bool, ShowWorkingRange )
+		/** Whether to show the view range */
+		SLATE_ARGUMENT( bool, ShowViewRange )
+		/** Whether to show the playback range */
+		SLATE_ARGUMENT( bool, ShowPlaybackRange )
+		/* Content to display inside the time range */
+		SLATE_DEFAULT_SLOT( FArguments, CenterContent )
 	SLATE_END_ARGS()
 
 	/**
@@ -24,47 +27,59 @@ public:
 	 * 
 	 * @param InArgs   A declaration from which to construct the widget
 	 */
-	void Construct( const FArguments& InArgs, TSharedRef<ITimeSliderController> InTimeSliderController );
+	void Construct( const FArguments& InArgs, TSharedRef<ITimeSliderController> InTimeSliderController, TSharedRef<INumericTypeInterface<float>> NumericTypeInterface );
+	
+protected:
 
-	float GetTimeSnapInterval() const;
+	float PlayStartTime() const;
+	float PlayEndTime() const;
+
+	TOptional<float> MinPlayStartTime() const;
+	TOptional<float> MaxPlayStartTime() const;
+	TOptional<float> MinPlayEndTime() const;
+	TOptional<float> MaxPlayEndTime() const;
+
+	void OnPlayStartTimeCommitted(float NewValue, ETextCommit::Type InTextCommit);
+	void OnPlayEndTimeCommitted(float NewValue, ETextCommit::Type InTextCommit);
+
+	void OnPlayStartTimeChanged(float NewValue);
+	void OnPlayEndTimeChanged(float NewValue);
+
+	FText PlayStartTimeTooltip() const;
+	FText PlayEndTimeTooltip() const;
 
-	using ITimeSlider::ToString;
-	
 protected:
-	float InTime() const;
-	float OutTime() const;
-	float StartTime() const;
-	float EndTime() const;
-
-	TOptional<float> MinInTime() const;
-	TOptional<float> MaxInTime() const;
-	TOptional<float> MinOutTime() const;
-	TOptional<float> MaxOutTime() const;
-	TOptional<float> MaxStartTime() const;
-	TOptional<float> MinEndTime() const;
-	
-	void OnStartTimeCommitted(float NewValue, ETextCommit::Type InTextCommit);
-	void OnEndTimeCommitted(float NewValue, ETextCommit::Type InTextCommit);
-	void OnInTimeCommitted(float NewValue, ETextCommit::Type InTextCommit);
-	void OnOutTimeCommitted(float NewValue, ETextCommit::Type InTextCommit);
 
-	void OnStartTimeChanged(float NewValue);
-	void OnEndTimeChanged(float NewValue);
-	void OnInTimeChanged(float NewValue);
-	void OnOutTimeChanged(float NewValue);
+	float ViewStartTime() const;
+	float ViewEndTime() const;
+
+	TOptional<float> MaxViewStartTime() const;
+	TOptional<float> MinViewEndTime() const;
+
+	void OnViewStartTimeCommitted(float NewValue, ETextCommit::Type InTextCommit);
+	void OnViewEndTimeCommitted(float NewValue, ETextCommit::Type InTextCommit);
+
+	void OnViewStartTimeChanged(float NewValue);
+	void OnViewEndTimeChanged(float NewValue);
+
+	FText ViewStartTimeTooltip() const;
+	FText ViewEndTimeTooltip() const;
+
+protected:
+
+	float WorkingStartTime() const;
+	float WorkingEndTime() const;
+
+	TOptional<float> MaxWorkingStartTime() const;
+	TOptional<float> MinWorkingEndTime() const;
 
-	FText InTimeTooltip() const;
-	FText OutTimeTooltip() const;
-	FText StartTimeTooltip() const;
-	FText EndTimeTooltip() const;
+	void OnWorkingStartTimeCommitted(float NewValue, ETextCommit::Type InTextCommit);
+	void OnWorkingEndTimeCommitted(float NewValue, ETextCommit::Type InTextCommit);
 
-	/** Convert the type to/from a string */
-	virtual FString ToString(const float& Value) const override;
-	virtual TOptional<float> FromString(const FString& InString, const float& ExistingValue) override;
+	void OnWorkingStartTimeChanged(float NewValue);
+	void OnWorkingEndTimeChanged(float NewValue);
 
 private:
 	TSharedPtr<ITimeSliderController> TimeSliderController;
-	TSharedPtr<STimeRangeSlider> TimeRangeSlider;
 	TAttribute<bool> ShowFrameNumbers;
-	TAttribute<float> TimeSnapInterval;
 };
\ No newline at end of file
diff --git a/Engine/Source/Editor/SequencerWidgets/Private/STimeRangeSlider.cpp b/Engine/Source/Editor/SequencerWidgets/Private/STimeRangeSlider.cpp
--- a/Engine/Source/Editor/SequencerWidgets/Private/STimeRangeSlider.cpp
+++ b/Engine/Source/Editor/SequencerWidgets/Private/STimeRangeSlider.cpp
@@ -2,7 +2,6 @@
 
 #include "SequencerWidgetsPrivatePCH.h"
 #include "STimeRangeSlider.h"
-#include "STimeRange.h"
 #include "SlateStyle.h"
 #include "EditorStyle.h"
 
@@ -13,11 +12,11 @@ namespace TimeRangeSliderConstants
 	const int32 HandleSize = 14;
 }
 
-void STimeRangeSlider::Construct( const FArguments& InArgs, TSharedRef<ITimeSliderController> InTimeSliderController, TSharedPtr<STimeRange> InTimeRange)
+void STimeRangeSlider::Construct( const FArguments& InArgs, TSharedRef<ITimeSliderController> InTimeSliderController)
 {
 	TimeSliderController = InTimeSliderController;
-	TimeRange = InTimeRange;
 	LastViewRange = TimeSliderController.Get()->GetViewRange();
+	TimeSnapInterval = InArgs._TimeSnapInterval;
 
 	ResetState();
 	ResetHoveredState();
@@ -154,14 +153,7 @@ FReply STimeRangeSlider::OnMouseMove( const FGeometry& MyGeometry, const FPointe
 	{
 		float DragDelta = ComputeDragDelta(MouseEvent, MyGeometry.Size.X);
 
-		float SnapInterval = KINDA_SMALL_NUMBER;
-		{
-			auto TimeRangePin = TimeRange.Pin();
-			if (TimeRangePin.IsValid())
-			{
-				SnapInterval = TimeRangePin->GetTimeSnapInterval();
-			}
-		}
+		const float SnapInterval = TimeSnapInterval.Get(1.f);
 
 		ITimeSliderController* TimeSliderControllerPtr = TimeSliderController.Get();
 		if (!TimeSliderControllerPtr)
diff --git a/Engine/Source/Editor/SequencerWidgets/Private/STimeRangeSlider.h b/Engine/Source/Editor/SequencerWidgets/Private/STimeRangeSlider.h
--- a/Engine/Source/Editor/SequencerWidgets/Private/STimeRangeSlider.h
+++ b/Engine/Source/Editor/SequencerWidgets/Private/STimeRangeSlider.h
@@ -2,16 +2,16 @@
 
 #pragma once
 
-class STimeRange;
-
 class STimeRangeSlider : public SCompoundWidget
 {
 public:
 	SLATE_BEGIN_ARGS(STimeRangeSlider){}
+		/* The time snap interval for the timeline */
+		SLATE_ARGUMENT( TAttribute<float>, TimeSnapInterval )
 		SLATE_DEFAULT_SLOT(FArguments, Content)
 	SLATE_END_ARGS()
 
-	void Construct( const FArguments& InArgs, TSharedRef<ITimeSliderController> InTimeSliderController, TSharedPtr<STimeRange> InTimeRange );
+	void Construct( const FArguments& InArgs, TSharedRef<ITimeSliderController> InTimeSliderController );
 
 	// SWidget interface
 	virtual int32 OnPaint( const FPaintArgs& Args, const FGeometry& AllottedGeometry, const FSlateRect& MyClippingRect, FSlateWindowElementList& OutDrawElements, int32 LayerId, const FWidgetStyle& InWidgetStyle, bool bParentEnabled ) const override;
@@ -51,6 +51,7 @@ private:
 	/* The in/out view range viewed before expansion */
 	TRange<float> LastViewRange;
 
+	TAttribute<float> TimeSnapInterval;
+
 	TSharedPtr<ITimeSliderController> TimeSliderController;
-	TWeakPtr<STimeRange> TimeRange;
 };
\ No newline at end of file
diff --git a/Engine/Source/Editor/SequencerWidgets/Private/SequencerWidgetsModule.cpp b/Engine/Source/Editor/SequencerWidgets/Private/SequencerWidgetsModule.cpp
--- a/Engine/Source/Editor/SequencerWidgets/Private/SequencerWidgetsModule.cpp
+++ b/Engine/Source/Editor/SequencerWidgets/Private/SequencerWidgetsModule.cpp
@@ -5,6 +5,7 @@
 #include "ModuleManager.h"
 #include "SSequencerTimeSlider.h"
 #include "STimeRange.h"
+#include "STimeRangeSlider.h"
 
 
 /**
@@ -30,12 +31,23 @@ public:
 			.MirrorLabels(bMirrorLabels);
 	}
 
-	TSharedRef<ITimeSlider> CreateTimeRange(const TSharedRef<ITimeSliderController>& InController, const TAttribute<EVisibility>& VisibilityDelegate, const TAttribute<bool>& ShowFrameNumbersDelegate, const TAttribute<float>& TimeSnapIntervalDelegate) override
+	TSharedRef<SWidget> CreateTimeRangeSlider( const TSharedRef<class ITimeSliderController>& InController, const TAttribute<float>& InTimeSnapIntervalDelegate ) override
 	{
-		return SNew(STimeRange, InController)
-			.Visibility(VisibilityDelegate)
-			.ShowFrameNumbers(ShowFrameNumbersDelegate)
-			.TimeSnapInterval(TimeSnapIntervalDelegate);
+		return SNew( STimeRangeSlider, InController )
+			.TimeSnapInterval( InTimeSnapIntervalDelegate );
+	}
+
+	TSharedRef<ITimeSlider> CreateTimeRange(const FTimeRangeArgs& InArgs, const TSharedRef<SWidget>& Content) override
+	{
+		return SNew( STimeRange, InArgs.Controller, InArgs.NumericTypeInterface)
+		.Visibility(InArgs.VisibilityDelegate)
+		.ShowFrameNumbers(InArgs.ShowFrameNumbersDelegate)
+		.ShowWorkingRange(!!(InArgs.ShowRanges & EShowRange::WorkingRange))
+		.ShowViewRange(!!(InArgs.ShowRanges & EShowRange::ViewRange))
+		.ShowPlaybackRange(!!(InArgs.ShowRanges & EShowRange::PlaybackRange))
+		[
+			Content
+		];
 	}
 
 public:
diff --git a/Engine/Source/Editor/SequencerWidgets/Public/ISequencerWidgetsModule.h b/Engine/Source/Editor/SequencerWidgets/Public/ISequencerWidgetsModule.h
--- a/Engine/Source/Editor/SequencerWidgets/Public/ISequencerWidgetsModule.h
+++ b/Engine/Source/Editor/SequencerWidgets/Public/ISequencerWidgetsModule.h
@@ -3,10 +3,51 @@
 #pragma once
 
 #include "ModuleInterface.h"
+#include "EnumClassFlags.h"
 
 
 class ITimeSlider;
+class ITimeSliderController;
 
+/** Enum denoting which time ranges to display on a time range */
+enum class EShowRange
+{
+	None			= 0,
+	WorkingRange	= 1 << 0,
+	ViewRange		= 1 << 1,
+	PlaybackRange	= 1 << 2,
+};
+ENUM_CLASS_FLAGS(EShowRange);
+
+/** Time range construction arguments */
+struct FTimeRangeArgs
+{
+	FTimeRangeArgs(
+		EShowRange InShowRanges,
+		TSharedRef<ITimeSliderController> InController,
+		TAttribute<EVisibility> InVisibilityDelegate,
+		TAttribute<bool> InShowFrameNumbersDelegate,
+		TSharedRef<INumericTypeInterface<float>> InNumericTypeInterface
+		)
+		: ShowRanges(InShowRanges)
+		, Controller(InController)
+		, VisibilityDelegate(InVisibilityDelegate)
+		, ShowFrameNumbersDelegate(InShowFrameNumbersDelegate)
+		, NumericTypeInterface(InNumericTypeInterface)
+	{
+	}
+	
+	/** Which ranges to show */
+	EShowRange ShowRanges;
+	/** A time slider controller */
+	TSharedRef<ITimeSliderController> Controller;
+	/** Visibility delegate */
+	TAttribute<EVisibility> VisibilityDelegate;
+	/** Whether to show frame numbers */
+	TAttribute<bool> ShowFrameNumbersDelegate;
+	/** Numeric type interface to use for frame<->time conversion */
+	TSharedRef<INumericTypeInterface<float>> NumericTypeInterface;
+};
 
 /**
  * The public interface of SequencerModule
@@ -16,7 +57,8 @@ class ISequencerWidgetsModule
 {
 public:
 
-	virtual TSharedRef<ITimeSlider> CreateTimeSlider(const TSharedRef<class ITimeSliderController>& InController, bool bMirrorLabels) = 0;
-	virtual TSharedRef<ITimeSlider> CreateTimeSlider(const TSharedRef<class ITimeSliderController>& InController, const TAttribute<EVisibility>& VisibilityDelegate, bool bMirrorLabels) = 0;
-	virtual TSharedRef<ITimeSlider> CreateTimeRange(const TSharedRef<class ITimeSliderController>& InController, const TAttribute<EVisibility>& VisibilityDelegate, const TAttribute<bool>& ShowFrameNumbersDelegate, const TAttribute<float>& TimeSnapIntervalDelegate) = 0;
+	virtual TSharedRef<ITimeSlider> CreateTimeSlider( const TSharedRef<ITimeSliderController>& InController, bool bMirrorLabels  ) = 0;
+	virtual TSharedRef<ITimeSlider> CreateTimeSlider( const TSharedRef<ITimeSliderController>& InController, const TAttribute<EVisibility>& VisibilityDelegate, bool bMirrorLabels  ) = 0;
+	virtual TSharedRef<SWidget> CreateTimeRangeSlider( const TSharedRef<class ITimeSliderController>& InController, const TAttribute<float>& InTimeSnapIntervalDelegate ) = 0;
+	virtual TSharedRef<ITimeSlider> CreateTimeRange( const FTimeRangeArgs& InArgs, const TSharedRef<SWidget>& Content ) = 0;
 };
diff --git a/Engine/Source/Editor/SequencerWidgets/Public/ITimeSlider.h b/Engine/Source/Editor/SequencerWidgets/Public/ITimeSlider.h
--- a/Engine/Source/Editor/SequencerWidgets/Public/ITimeSlider.h
+++ b/Engine/Source/Editor/SequencerWidgets/Public/ITimeSlider.h
@@ -3,6 +3,7 @@
 #pragma once
 
 #include "Editor/Sequencer/Public/ISequencerInputHandler.h"
+#include "IMovieScenePlayer.h"
 
 class USequencerSettings;
 
@@ -70,30 +71,58 @@ struct FTimeSliderArgs
 
 	/** The scrub position */
 	TAttribute<float> ScrubPosition;
+
 	/** View time range */
 	TAttribute< FAnimatedRange > ViewRange;
+
 	/** Clamp time range */
 	TAttribute< FAnimatedRange > ClampRange;
+
 	/** Called when the scrub position changes */
 	FOnScrubPositionChanged OnScrubPositionChanged;
+
 	/** Called right before the scrubber begins to move */
 	FSimpleDelegate OnBeginScrubberMovement;
+
 	/** Called right after the scrubber handle is released by the user */
 	FSimpleDelegate OnEndScrubberMovement;
+
 	/** Called when the view range changes */
 	FOnViewRangeChanged OnViewRangeChanged;
+
 	/** Called when the clamp range changes */
 	FOnRangeChanged OnClampRangeChanged;
+
+	/** Attribute defining the playback range for this controller */
+	TAttribute<TRange<float>> InOutRange;
+
+	/** Delegate that is called when the in/out range wants to change */
+	FOnRangeChanged OnInOutRangeChanged;
+
+	/** Called right before the in/out range starts to be dragged */
+	FSimpleDelegate OnBeginInOutRangeDrag;
+
+	/** Called right after the in/out range has finished being dragged */
+	FSimpleDelegate OnEndInOutRangeDrag;
+
 	/** Attribute defining the playback range for this controller */
 	TAttribute<TRange<float>> PlaybackRange;
+
 	/** Delegate that is called when the playback range wants to change */
 	FOnRangeChanged OnPlaybackRangeChanged;
+
 	/** Called right before the playback range starts to be dragged */
 	FSimpleDelegate OnBeginPlaybackRangeDrag;
+
 	/** Called right after the playback range has finished being dragged */
 	FSimpleDelegate OnEndPlaybackRangeDrag;
+
+	/** Round the scrub position to an integer during playback */
+	TAttribute<EMovieScenePlayerStatus::Type> PlaybackStatus;
+
 	/** If we are allowed to zoom */
 	bool AllowZoom;
+
 	/** User-supplied settings object */
 	USequencerSettings* Settings;
 };
@@ -111,6 +140,9 @@ public:
 	/** Get the current clamp range for this controller */
 	virtual FAnimatedRange GetClampRange() const { return FAnimatedRange(); }
 
+	/** Get the current play range for this controller */
+	virtual TRange<float> GetPlayRange() const { return TRange<float>(); }
+
 	/** Convert a time to a frame */
 	virtual int32 TimeToFrame(float Time) const { return 1; }
 
@@ -133,6 +165,14 @@ public:
 	 * @param NewRangeMax		The new upper bound of the clamp range
 	 */
 	virtual void SetClampRange( float NewRangeMin, float NewRangeMax) {}
+
+	/**
+	 * Set a new playback range based on a min, max
+	 * 
+	 * @param NewRangeMin		The new lower bound of the playback range
+	 * @param NewRangeMax		The new upper bound of the playback range
+	 */
+	virtual void SetPlayRange( float NewRangeMin, float NewRangeMax ) {}
 };
 
 /**
diff --git a/Engine/Source/Editor/SequencerWidgets/SequencerWidgets.Build.cs b/Engine/Source/Editor/SequencerWidgets/SequencerWidgets.Build.cs
--- a/Engine/Source/Editor/SequencerWidgets/SequencerWidgets.Build.cs
+++ b/Engine/Source/Editor/SequencerWidgets/SequencerWidgets.Build.cs
@@ -15,6 +15,7 @@ public class SequencerWidgets : ModuleRules
                 "EditorStyle",
 				"Engine",
                 "InputCore",
+                "MovieScene",
 				"Slate",
 				"SlateCore",
 				"UnrealEd"
diff --git a/Engine/Source/Editor/UnrealEd/Classes/Settings/LevelEditorViewportSettings.h b/Engine/Source/Editor/UnrealEd/Classes/Settings/LevelEditorViewportSettings.h
--- a/Engine/Source/Editor/UnrealEd/Classes/Settings/LevelEditorViewportSettings.h
+++ b/Engine/Source/Editor/UnrealEd/Classes/Settings/LevelEditorViewportSettings.h
@@ -88,6 +88,7 @@ struct UNREALED_API FLevelEditorViewportInstanceSettings
 		, bIsRealtime(false)
 		, bShowFPS_DEPRECATED(false)
 		, bShowStats(false)
+		, bShowFullToolbar(true)
 	{ }
 
 	/** The viewport type */
@@ -147,6 +148,10 @@ struct UNREALED_API FLevelEditorViewportInstanceSettings
 	/* Whether viewport statistics should be enabled by default. */
 	UPROPERTY(config)
 	TArray<FString> EnabledStats;
+
+	/** When enabled, the full viewport toolbar will be shown. When disabled, a compact toolbar is used. */
+	UPROPERTY(EditAnywhere, config, Category=LookAndFeel)
+	bool bShowFullToolbar;
 };
 
 
diff --git a/Engine/Source/Editor/UnrealEd/Private/EditorViewportClient.cpp b/Engine/Source/Editor/UnrealEd/Private/EditorViewportClient.cpp
--- a/Engine/Source/Editor/UnrealEd/Private/EditorViewportClient.cpp
+++ b/Engine/Source/Editor/UnrealEd/Private/EditorViewportClient.cpp
@@ -857,7 +857,14 @@ FSceneView* FEditorViewportClient::CalcSceneView(FSceneViewFamily* ViewFamily)
 
 	View->StartFinalPostprocessSettings(ViewLocation);
 
-	OverridePostProcessSettings( *View );
+	if (bUseControllingActorViewInfo)
+	{
+		View->OverridePostProcessSettings(ControllingActorViewInfo.PostProcessSettings, ControllingActorViewInfo.PostProcessBlendWeight);
+	}
+	else
+	{
+		OverridePostProcessSettings(*View);
+	}
 
 	View->EndFinalPostprocessSettings(ViewInitOptions);
 
diff --git a/Engine/Source/Editor/UnrealEd/Private/FeedbackContextEditor.cpp b/Engine/Source/Editor/UnrealEd/Private/FeedbackContextEditor.cpp
--- a/Engine/Source/Editor/UnrealEd/Private/FeedbackContextEditor.cpp
+++ b/Engine/Source/Editor/UnrealEd/Private/FeedbackContextEditor.cpp
@@ -363,13 +363,6 @@ void FFeedbackContextEditor::Serialize( const TCHAR* V, ELogVerbosity::Type Verb
 
 void FFeedbackContextEditor::StartSlowTask( const FText& Task, bool bShowCancelButton )
 {
-	// If we are in PIE, then do not show slow dialog window, it can cause crashes
-	// if there are UMG gadgets open.  This is because they will get a message during
-	// an asset load and that causes an assert.
-	if (GEditor != nullptr && GEditor->PlayWorld != nullptr)
-	{
-		return;
-	}
 	FFeedbackContext::StartSlowTask( Task, bShowCancelButton );
 
 	// Attempt to parent the slow task window to the slate main frame
@@ -596,3 +589,8 @@ void FFeedbackContextEditor::CloseBuildProgressWindow()
 		BuildProgressWidget.Reset();	
 	}
 }
+
+bool FFeedbackContextEditor::IsPlayingInEditor() const
+{
+	return (GIsPlayInEditorWorld || (GEditor && GEditor->PlayWorld != nullptr));
+}
diff --git a/Engine/Source/Editor/UnrealEd/Private/SCurveEditor.cpp b/Engine/Source/Editor/UnrealEd/Private/SCurveEditor.cpp
--- a/Engine/Source/Editor/UnrealEd/Private/SCurveEditor.cpp
+++ b/Engine/Source/Editor/UnrealEd/Private/SCurveEditor.cpp
@@ -80,11 +80,6 @@ void SCurveEditor::Construct(const FArguments& InArgs)
 
 	MovementAxisLock = EMovementAxisLock::None;
 
-	//Simple r/g/b for now
-	CurveColors.Add(FLinearColor(1.0f, 0.0f, 0.0f));
-	CurveColors.Add(FLinearColor(0.0f, 1.0f, 0.0f));
-	CurveColors.Add(FLinearColor(0.05f, 0.05f, 1.0f));
-
 	TransactionIndex = -1;
 
 	ReduceTolerance = 0.001;
@@ -1194,7 +1189,7 @@ void SCurveEditor::SetCurveOwner(FCurveOwnerInterface* InCurveOwner, bool bCanEd
 		int curveIndex = 0;
 		for (auto CurveInfo : CurveOwner->GetCurves())
 		{
-			CurveViewModels.Add(TSharedPtr<FCurveViewModel>(new FCurveViewModel(CurveInfo, CurveColors[curveIndex % CurveColors.Num()], !bCanEdit)));
+			CurveViewModels.Add(TSharedPtr<FCurveViewModel>(new FCurveViewModel(CurveInfo, CurveOwner->GetCurveColor(CurveInfo), !bCanEdit)));
 			curveIndex++;
 		}
 		CurveOwner->MakeTransactional();
@@ -2105,28 +2100,29 @@ SCurveEditor::FSelectedCurveKey SCurveEditor::HitTestKeys(const FGeometry& InMyG
 
 	if( AreCurvesVisible() )
 	{
+		bool bAnyCurveViewModelsSelected = AnyCurveViewModelsSelected();
+
 		FTrackScaleInfo ScaleInfo(ViewMinInput.Get(),  ViewMaxInput.Get(), ViewMinOutput.Get(), ViewMaxOutput.Get(), InMyGeometry.Size);
 
 		const FVector2D HitPosition = InMyGeometry.AbsoluteToLocal( HitScreenPosition );
 
-
 		for(auto CurveViewModel : CurveViewModels)
 		{
-			if (!CurveViewModel->bIsLocked && CurveViewModel->bIsVisible)
-		{
-				FRichCurve* Curve = CurveViewModel->CurveInfo.CurveToEdit;
-			if(Curve != NULL)
+			if (IsCurveSelectable(CurveViewModel))
 			{
-				for (auto It(Curve->GetKeyHandleIterator()); It; ++It)
+				FRichCurve* Curve = CurveViewModel->CurveInfo.CurveToEdit;
+				if(Curve != NULL)
 				{
-					float KeyScreenX = ScaleInfo.InputToLocalX(Curve->GetKeyTime(It.Key()));
-					float KeyScreenY = ScaleInfo.OutputToLocalY(Curve->GetKeyValue(It.Key()));
-
-					if(	HitPosition.X > (KeyScreenX - (0.5f * CONST_KeySize.X)) && 
-						HitPosition.X < (KeyScreenX + (0.5f * CONST_KeySize.X)) &&
-						HitPosition.Y > (KeyScreenY - (0.5f * CONST_KeySize.Y)) &&
-						HitPosition.Y < (KeyScreenY + (0.5f * CONST_KeySize.Y)) )
+					for (auto It(Curve->GetKeyHandleIterator()); It; ++It)
 					{
+						float KeyScreenX = ScaleInfo.InputToLocalX(Curve->GetKeyTime(It.Key()));
+						float KeyScreenY = ScaleInfo.OutputToLocalY(Curve->GetKeyValue(It.Key()));
+
+						if(	HitPosition.X > (KeyScreenX - (0.5f * CONST_KeySize.X)) && 
+							HitPosition.X < (KeyScreenX + (0.5f * CONST_KeySize.X)) &&
+							HitPosition.Y > (KeyScreenY - (0.5f * CONST_KeySize.Y)) &&
+							HitPosition.Y < (KeyScreenY + (0.5f * CONST_KeySize.Y)) )
+						{
 							return  FSelectedCurveKey(Curve, It.Key());
 						}
 					}
@@ -2820,19 +2816,28 @@ TSharedPtr<FCurveViewModel> SCurveEditor::HitTestCurves(  const FGeometry& InMyG
 	return TSharedPtr<FCurveViewModel>();
 }
 
+bool SCurveEditor::IsCurveSelectable(TSharedPtr<FCurveViewModel> CurveViewModel) const
+{		
+	bool bAnyCurveViewModelsSelected = AnyCurveViewModelsSelected();
+	bool bDisabled = bAnyCurveViewModelsSelected && !CurveViewModel->bIsSelected;
+
+	return !CurveViewModel->bIsLocked && CurveViewModel->bIsVisible && !bDisabled;
+}
+
 SCurveEditor::FSelectedTangent SCurveEditor::HitTestCubicTangents( const FGeometry& InMyGeometry, const FVector2D& HitScreenPosition )
 {
 	FSelectedTangent Tangent;
 
 	if( AreCurvesVisible() )
 	{
+
 		FTrackScaleInfo ScaleInfo(ViewMinInput.Get(),  ViewMaxInput.Get(), ViewMinOutput.Get(), ViewMaxOutput.Get(), InMyGeometry.Size);
 
 		const FVector2D HitPosition = InMyGeometry.AbsoluteToLocal( HitScreenPosition);
 
 		for (auto CurveViewModel : CurveViewModels)
 		{
-			if (!CurveViewModel->bIsLocked && CurveViewModel->bIsVisible)
+			if (IsCurveSelectable(CurveViewModel))
 			{
 				FRichCurve* Curve = CurveViewModel->CurveInfo.CurveToEdit;
 				if (Curve != NULL)
@@ -2981,6 +2986,7 @@ void SCurveEditor::OnFlattenOrStraightenTangents(bool bFlattenTangents)
 
 			Key.Curve->GetKey(Key.KeyHandle).LeaveTangent = LeaveTangent;
 			Key.Curve->GetKey(Key.KeyHandle).ArriveTangent = ArriveTangent;
+			Key.Curve->GetKey(Key.KeyHandle).TangentMode = RCTM_User;
 		}
 				
 		for(auto It = SelectedTangents.CreateIterator();It;++It)
@@ -3004,6 +3010,7 @@ void SCurveEditor::OnFlattenOrStraightenTangents(bool bFlattenTangents)
 
 			Tangent.Key.Curve->GetKey(Tangent.Key.KeyHandle).LeaveTangent = LeaveTangent;
 			Tangent.Key.Curve->GetKey(Tangent.Key.KeyHandle).ArriveTangent = ArriveTangent;
+			Tangent.Key.Curve->GetKey(Tangent.Key.KeyHandle).TangentMode = RCTM_User;
 		}
 
 		TArray<FRichCurveEditInfo> ChangedCurveEditInfos;
@@ -3399,7 +3406,7 @@ TArray<SCurveEditor::FSelectedCurveKey> SCurveEditor::GetEditableKeysWithinMarqu
 		FTrackScaleInfo ScaleInfo(ViewMinInput.Get(), ViewMaxInput.Get(), ViewMinOutput.Get(), ViewMaxOutput.Get(), InMyGeometry.Size);
 		for (auto CurveViewModel : CurveViewModels)
 		{
-			if (!CurveViewModel->bIsLocked && CurveViewModel->bIsVisible)
+			if (IsCurveSelectable(CurveViewModel))
 			{
 				FRichCurve* Curve = CurveViewModel->CurveInfo.CurveToEdit;
 				if (Curve != NULL)
@@ -3437,7 +3444,7 @@ TArray<SCurveEditor::FSelectedTangent> SCurveEditor::GetEditableTangentsWithinMa
 		FTrackScaleInfo ScaleInfo(ViewMinInput.Get(), ViewMaxInput.Get(), ViewMinOutput.Get(), ViewMaxOutput.Get(), InMyGeometry.Size);
 		for (auto CurveViewModel : CurveViewModels)
 		{
-			if (!CurveViewModel->bIsLocked && CurveViewModel->bIsVisible)
+			if (IsCurveSelectable(CurveViewModel))
 			{
 				FRichCurve* Curve = CurveViewModel->CurveInfo.CurveToEdit;
 				if (Curve != NULL)
diff --git a/Engine/Source/Editor/UnrealEd/Public/CurveEditorSettings.cpp b/Engine/Source/Editor/UnrealEd/Public/CurveEditorSettings.cpp
--- a/Engine/Source/Editor/UnrealEd/Public/CurveEditorSettings.cpp
+++ b/Engine/Source/Editor/UnrealEd/Public/CurveEditorSettings.cpp
@@ -8,7 +8,7 @@ UCurveEditorSettings::UCurveEditorSettings( const FObjectInitializer& ObjectInit
 {
 	bAutoFrameCurveEditor = true;
 	bShowCurveEditorCurveToolTips = true;
-	CurveVisibility = ECurveEditorCurveVisibility::AllCurves;
+	CurveVisibility = ECurveEditorCurveVisibility::SelectedCurves;
 	TangentVisibility = ECurveEditorTangentVisibility::SelectedKeys;
 }
 
diff --git a/Engine/Source/Editor/UnrealEd/Public/FeedbackContextEditor.h b/Engine/Source/Editor/UnrealEd/Public/FeedbackContextEditor.h
--- a/Engine/Source/Editor/UnrealEd/Public/FeedbackContextEditor.h
+++ b/Engine/Source/Editor/UnrealEd/Public/FeedbackContextEditor.h
@@ -30,6 +30,7 @@ public:
 	virtual void StartSlowTask( const FText& Task, bool bShowCancelButton=false ) override;
 	virtual void FinalizeSlowTask( ) override;
 	virtual void ProgressReported( const float TotalProgressInterp, FText DisplayMessage ) override;
+	virtual bool IsPlayingInEditor() const override;
 
 	void SetContext( FContextSupplier* InSupplier ) override {}
 
diff --git a/Engine/Source/Editor/UnrealEd/Public/SCurveEditor.h b/Engine/Source/Editor/UnrealEd/Public/SCurveEditor.h
--- a/Engine/Source/Editor/UnrealEd/Public/SCurveEditor.h
+++ b/Engine/Source/Editor/UnrealEd/Public/SCurveEditor.h
@@ -539,6 +539,8 @@ private:
 	bool IsGradientEditorVisible() const { return bIsGradientEditorVisible; }
 	bool IsLinearColorCurve() const;
 
+	bool IsCurveSelectable(TSharedPtr<FCurveViewModel> CurveViewModel) const;
+
 	FVector2D SnapLocation(FVector2D InLocation);
 
 	FText GetIsCurveVisibleToolTip(TSharedPtr<FCurveViewModel> CurveViewModel) const;
@@ -608,8 +610,6 @@ private:
 
 	/** Interface for curve supplier */
 	FCurveOwnerInterface*		CurveOwner;
-	/** Color for each curve */
-	TArray<FLinearColor>		CurveColors;
 
 	/** If we should draw the curve */
 	bool				bDrawCurve;
diff --git a/Engine/Source/Editor/UnrealEd/UnrealEd.Build.cs b/Engine/Source/Editor/UnrealEd/UnrealEd.Build.cs
--- a/Engine/Source/Editor/UnrealEd/UnrealEd.Build.cs
+++ b/Engine/Source/Editor/UnrealEd/UnrealEd.Build.cs
@@ -224,6 +224,12 @@ public class UnrealEd : ModuleRules
 			DynamicallyLoadedModuleNames.Add("AndroidPlatformEditor");
 		}
 
+        if (Target.Type == TargetRules.TargetType.Editor)
+        {
+            DynamicallyLoadedModuleNames.Add("SequenceRecorder");
+            PrivateIncludePathModuleNames.Add("SequenceRecorder");
+        }
+
 		CircularlyReferencedDependentModules.AddRange(
 			new string[] 
 			{
diff --git a/Engine/Source/Runtime/AppFramework/Private/Widgets/Colors/SColorPicker.cpp b/Engine/Source/Runtime/AppFramework/Private/Widgets/Colors/SColorPicker.cpp
--- a/Engine/Source/Runtime/AppFramework/Private/Widgets/Colors/SColorPicker.cpp
+++ b/Engine/Source/Runtime/AppFramework/Private/Widgets/Colors/SColorPicker.cpp
@@ -16,7 +16,7 @@
 
 
 /** A default window size for the color picker which looks nice */
-const FVector2D SColorPicker::DEFAULT_WINDOW_SIZE = FVector2D(308, 458);
+const FVector2D SColorPicker::DEFAULT_WINDOW_SIZE = FVector2D(441, 537);
 
 /** The max time allowed for updating before we shut off auto-updating */
 const double SColorPicker::MAX_ALLOWED_UPDATE_TIME = 0.1;
diff --git a/Engine/Source/Runtime/CinematicCamera/Private/CameraRig_Crane.cpp b/Engine/Source/Runtime/CinematicCamera/Private/CameraRig_Crane.cpp
--- a/Engine/Source/Runtime/CinematicCamera/Private/CameraRig_Crane.cpp
+++ b/Engine/Source/Runtime/CinematicCamera/Private/CameraRig_Crane.cpp
@@ -22,7 +22,7 @@ ACameraRig_Crane::ACameraRig_Crane(const FObjectInitializer& ObjectInitializer)
 	// Setup camera defaults
 	CraneYawControl = CreateDefaultSubobject<USceneComponent>(TEXT("CraneYawControl"));
 	CraneYawControl->AttachParent = TransformComponent;
-	CraneYawControl->RelativeLocation = FVector(0.f, 0.f, 100.f);			// pivot height off the ground
+	CraneYawControl->RelativeLocation = FVector(0.f, 0.f, 70.f);			// pivot height off the ground
 	CraneYawControl->RelativeRotation = FRotator(0.f, CraneYaw, 0.f);
 
 	CranePitchControl = CreateDefaultSubobject<USceneComponent>(TEXT("CranePitchControl"));
@@ -31,12 +31,12 @@ ACameraRig_Crane::ACameraRig_Crane(const FObjectInitializer& ObjectInitializer)
 
 	CraneCameraMount = CreateDefaultSubobject<USceneComponent>(TEXT("CraneCameraMount"));
 	CraneCameraMount->AttachParent = CranePitchControl;
-	CraneCameraMount->RelativeLocation = FVector(CraneArmLength, 0.f, -35.f);			// negative z == underslung mount
+	CraneCameraMount->RelativeLocation = FVector(CraneArmLength, 0.f, -15.f);			// negative z == underslung mount
 
 	// preview meshes
 	if (!IsRunningCommandlet() && !IsRunningDedicatedServer())
 	{
-		static ConstructorHelpers::FObjectFinder<UStaticMesh> CraneBaseMesh(TEXT("/Engine/BasicShapes/Cone"));
+		static ConstructorHelpers::FObjectFinder<UStaticMesh> CraneBaseMesh(TEXT("/Engine/EditorMeshes/Camera/SM_CraneRig_Base.SM_CraneRig_Base"));
 		PreviewMesh_CraneBase = CreateOptionalDefaultSubobject<UStaticMeshComponent>(TEXT("PreviewMesh_CraneBase"));
 		if (PreviewMesh_CraneBase)
 		{
@@ -48,11 +48,10 @@ ACameraRig_Crane::ACameraRig_Crane(const FObjectInitializer& ObjectInitializer)
 			PreviewMesh_CraneBase->CastShadow = false;
 			PreviewMesh_CraneBase->PostPhysicsComponentTick.bCanEverTick = false;
 
-			PreviewMesh_CraneBase->AttachParent = RootComponent;		// sibling of yawcontrol
-			PreviewMesh_CraneBase->RelativeLocation = FVector(0.f, 0.f, 50.f);
+			PreviewMesh_CraneBase->AttachParent = TransformComponent;		// sibling of yawcontrol
 		}
 
-		static ConstructorHelpers::FObjectFinder<UStaticMesh> CraneArmMesh(TEXT("/Engine/BasicShapes/Cube"));
+		static ConstructorHelpers::FObjectFinder<UStaticMesh> CraneArmMesh(TEXT("/Engine/EditorMeshes/Camera/SM_CraneRig_Arm.SM_CraneRig_Arm"));
 		PreviewMesh_CraneArm = CreateOptionalDefaultSubobject<UStaticMeshComponent>(TEXT("PreviewMesh_CraneArm"));
 		if (PreviewMesh_CraneArm)
 		{
@@ -65,35 +64,67 @@ ACameraRig_Crane::ACameraRig_Crane(const FObjectInitializer& ObjectInitializer)
 			PreviewMesh_CraneArm->PostPhysicsComponentTick.bCanEverTick = false;
 			
 			PreviewMesh_CraneArm->AttachParent = CranePitchControl;		// sibling of the mount
-			PreviewMesh_CraneArm->RelativeScale3D = FVector(0.2f, 0.2f, 0.2f);
+			PreviewMesh_CraneArm->RelativeRotation = FRotator(0.f, 90.f, 0.f);
+			PreviewMesh_CraneArm->RelativeLocation = FVector(0.f, 0.f, 52.f);
+			PreviewMesh_CraneArm->RelativeScale3D = FVector(0.7f, 0.7f, 0.7f);
+		}
+		
+		static ConstructorHelpers::FObjectFinder<UStaticMesh> CraneArmMount(TEXT("/Engine/EditorMeshes/Camera/SM_CraneRig_Mount.SM_CraneRig_Mount"));
+		PreviewMesh_CraneMount = CreateOptionalDefaultSubobject<UStaticMeshComponent>(TEXT("PreviewMesh_CraneMount"));
+		if (PreviewMesh_CraneMount)
+		{
+			PreviewMesh_CraneMount->SetStaticMesh(CraneArmMount.Object);
+			PreviewMesh_CraneMount->AlwaysLoadOnClient = false;
+			PreviewMesh_CraneMount->AlwaysLoadOnServer = false;
+			PreviewMesh_CraneMount->SetCollisionProfileName(UCollisionProfile::NoCollision_ProfileName);
+			PreviewMesh_CraneMount->bHiddenInGame = true;
+			PreviewMesh_CraneMount->CastShadow = false;
+			PreviewMesh_CraneMount->PostPhysicsComponentTick.bCanEverTick = false;
+
+			PreviewMesh_CraneMount->AttachParent = CraneCameraMount;
+			PreviewMesh_CraneMount->RelativeRotation = FRotator(0.f, 90.f, 0.f);
+			PreviewMesh_CraneMount->RelativeLocation = FVector(0.f, 0.f, 57.f);
+			PreviewMesh_CraneMount->RelativeScale3D = FVector(0.7f, 0.7f, 0.7f);
+		}
+		
+		static ConstructorHelpers::FObjectFinder<UStaticMesh> CraneArmCW(TEXT("/Engine/EditorMeshes/Camera/SM_CraneRig_Body.SM_CraneRig_Body"));
+		PreviewMesh_CraneCounterWeight = CreateOptionalDefaultSubobject<UStaticMeshComponent>(TEXT("PreviewMesh_CraneCounterWeight"));
+		if (PreviewMesh_CraneCounterWeight)
+		{
+			PreviewMesh_CraneCounterWeight->SetStaticMesh(CraneArmCW.Object);
+			PreviewMesh_CraneCounterWeight->AlwaysLoadOnClient = false;
+			PreviewMesh_CraneCounterWeight->AlwaysLoadOnServer = false;
+			PreviewMesh_CraneCounterWeight->SetCollisionProfileName(UCollisionProfile::NoCollision_ProfileName);
+			PreviewMesh_CraneCounterWeight->bHiddenInGame = true;
+			PreviewMesh_CraneCounterWeight->CastShadow = false;
+			PreviewMesh_CraneCounterWeight->PostPhysicsComponentTick.bCanEverTick = false;
+
+			PreviewMesh_CraneCounterWeight->AttachParent = CranePitchControl;
+			PreviewMesh_CraneCounterWeight->RelativeRotation = FRotator(0.f, 90.f, 0.f);
+			PreviewMesh_CraneCounterWeight->RelativeScale3D = FVector(0.7f, 0.7f, 0.7f);
 		}
 
 		UpdatePreviewMeshes();
 	}
 }
 
-static const float CraneArmMesh_CounterweightOverhang = 100.f;
+static const float CraneArmMesh_CounterweightOverhang = 0.f;
+static const float CraneArmMesh_DefaultMeshSize = 29.f * 0.7f;		// size of the mesh in the dimension that will stretch (accounting for the 0.7 scale)
 
 void ACameraRig_Crane::UpdatePreviewMeshes()
 {
 	if (PreviewMesh_CraneArm)
 	{
 		// note to explain the math here:
-		// crane arm mesh is current a 100x100 cube. we want it to overhang the base pivot by 100cm
-		// (chosen arbitrarily)
+		// this tailored to the dimensions of the asset (29cm in scaling direction, pivot at the edge)
 
 		float const TotalCraneArmMeshLength = CraneArmMesh_CounterweightOverhang + CraneArmLength;
-		float const CraneArmXScale = TotalCraneArmMeshLength / 100.f;
-		float const CraneArmXOffset = TotalCraneArmMeshLength * 0.5f - CraneArmMesh_CounterweightOverhang;
+		float const CraneArmYScale = TotalCraneArmMeshLength / CraneArmMesh_DefaultMeshSize;
 
 		FVector NewRelScale3D = PreviewMesh_CraneArm->RelativeScale3D;
-		NewRelScale3D.X = CraneArmXScale;
-
-		FVector NewRelLoc = PreviewMesh_CraneArm->RelativeLocation;
-		NewRelLoc.X = CraneArmXOffset;
+		NewRelScale3D.Y = CraneArmYScale * 0.7f;
 
 		PreviewMesh_CraneArm->SetRelativeScale3D(NewRelScale3D);
-		PreviewMesh_CraneArm->SetRelativeLocation(NewRelLoc);
 	}
 }
 
diff --git a/Engine/Source/Runtime/CinematicCamera/Private/CameraRig_Rail.cpp b/Engine/Source/Runtime/CinematicCamera/Private/CameraRig_Rail.cpp
--- a/Engine/Source/Runtime/CinematicCamera/Private/CameraRig_Rail.cpp
+++ b/Engine/Source/Runtime/CinematicCamera/Private/CameraRig_Rail.cpp
@@ -2,28 +2,170 @@
 
 #include "CinematicCameraPrivate.h"
 #include "CameraRig_Rail.h"
+#include "Components/SplineComponent.h"
+#include "Components/SplineMeshComponent.h"
 
 #define LOCTEXT_NAMESPACE "CameraRig_Rail"
 
-// #TODO WRITEME
-
 ACameraRig_Rail::ACameraRig_Rail(const FObjectInitializer& ObjectInitializer)
 : Super(ObjectInitializer)
 {
+	PrimaryActorTick.bCanEverTick = true;
+	PrimaryActorTick.bStartWithTickEnabled = true;
+
+	CurrentPositionOnRail = 0.f;
+
+	// Create components
+	TransformComponent = CreateDefaultSubobject<USceneComponent>(TEXT("TransformComponent"));
+
+	// Make the scene component the root component
+	RootComponent = TransformComponent;
+
+	RailSplineComponent = CreateDefaultSubobject<USplineComponent>(TEXT("RailSplineComponent"));
+	RailSplineComponent->AttachParent = TransformComponent;
+
+	RailCameraMount = CreateDefaultSubobject<USceneComponent>(TEXT("RailCameraMount"));
+	RailCameraMount->AttachParent = RailSplineComponent;
+
+	if (!IsRunningCommandlet() && !IsRunningDedicatedServer())
+	{
+		static ConstructorHelpers::FObjectFinder<UStaticMesh> RailMesh(TEXT("/Engine/EditorMeshes/Camera/SM_RailRig_Track.SM_RailRig_Track"));
+		PreviewRailStaticMesh = RailMesh.Object;
+
+		static ConstructorHelpers::FObjectFinder<UStaticMesh> MountMesh(TEXT("/Engine/EditorMeshes/Camera/SM_RailRig_Mount.SM_RailRig_Mount"));
+		PreviewMesh_Mount = CreateOptionalDefaultSubobject<UStaticMeshComponent>(TEXT("PreviewMesh_Mount"));
+		if (PreviewMesh_Mount)
+		{
+			PreviewMesh_Mount->SetStaticMesh(MountMesh.Object);
+			PreviewMesh_Mount->AlwaysLoadOnClient = false;
+			PreviewMesh_Mount->AlwaysLoadOnServer = false;
+			PreviewMesh_Mount->SetCollisionProfileName(UCollisionProfile::NoCollision_ProfileName);
+			PreviewMesh_Mount->bHiddenInGame = true;
+			PreviewMesh_Mount->CastShadow = false;
+			PreviewMesh_Mount->PostPhysicsComponentTick.bCanEverTick = false;
+
+			PreviewMesh_Mount->AttachParent = RailCameraMount;
+		}
+	}
 }
 
-#if WITH_EDITOR
-void ACameraRig_Rail::PostEditChangeProperty(FPropertyChangedEvent& PropertyChangedEvent)
+USplineMeshComponent* ACameraRig_Rail::CreateSplinePreviewSegment()
 {
-	Super::PostEditChangeProperty(PropertyChangedEvent);
+	USplineMeshComponent* const Segment = NewObject<USplineMeshComponent>(this);
+	if (Segment)
+	{
+		Segment->SetStaticMesh(PreviewRailStaticMesh);
+		Segment->SetMobility(EComponentMobility::Movable);
+		Segment->AlwaysLoadOnClient = false;
+		Segment->AlwaysLoadOnServer = false;
+		Segment->SetCollisionProfileName(UCollisionProfile::NoCollision_ProfileName);
+		Segment->bHiddenInGame = true;
+		Segment->CastShadow = false;
+		Segment->PostPhysicsComponentTick.bCanEverTick = false;
+	 
+		Segment->AttachParent = TransformComponent;
+		Segment->RegisterComponent();
+	}
+
+	return Segment;
+}
+
+void ACameraRig_Rail::UpdatePreviewMeshes()
+{
+	if (PreviewRailStaticMesh)
+	{
+		int32 const NumSplinePoints = RailSplineComponent->GetNumberOfSplinePoints();
+		int32 const NumNeededPreviewMeshes = NumSplinePoints - 1;
+
+		// make sure our preview mesh array is correctly sized and populated
+		{
+			int32 const NumExistingPreviewMeshes = PreviewRailMeshSegments.Num();
+			if (NumExistingPreviewMeshes > NumNeededPreviewMeshes)
+			{
+				// we have too many meshes, remove some
+				int32 const NumToRemove = NumExistingPreviewMeshes - NumNeededPreviewMeshes;
+				for (int Idx = 0; Idx < NumToRemove; ++Idx)
+				{
+					USplineMeshComponent* const ElementToRemove = PreviewRailMeshSegments.Pop();
+					ElementToRemove->UnregisterComponent();
+				}
+			}
+			else if (NumExistingPreviewMeshes < NumNeededPreviewMeshes)
+			{
+				int32 const NumToAdd = NumNeededPreviewMeshes - NumExistingPreviewMeshes;
+
+				for (int32 Idx = 0; Idx < NumToAdd; ++Idx)
+				{
+					USplineMeshComponent* PreviewMesh = CreateSplinePreviewSegment();
+					PreviewRailMeshSegments.Add(PreviewMesh);
+				}
+			}
+			check(PreviewRailMeshSegments.Num() == NumNeededPreviewMeshes);
+		}
+
+		for (int PtIdx = 0; PtIdx < NumSplinePoints - 1; ++PtIdx)
+		{
+			FVector StartLoc, StartTangent, EndLoc, EndTangent;
+			RailSplineComponent->GetLocationAndTangentAtSplinePoint(PtIdx, StartLoc, StartTangent, ESplineCoordinateSpace::Local);
+			RailSplineComponent->GetLocationAndTangentAtSplinePoint(PtIdx + 1, EndLoc, EndTangent, ESplineCoordinateSpace::Local);
+
+			USplineMeshComponent* const SplineMeshComp = PreviewRailMeshSegments[PtIdx];
+			SplineMeshComp->SetStartScale(FVector2D(0.1f, 0.1f), false);
+			SplineMeshComp->SetEndScale(FVector2D(0.1f, 0.1f), false);
+			SplineMeshComp->SetForwardAxis(ESplineMeshAxis::Z);
+			SplineMeshComp->SetStartAndEnd(StartLoc, StartTangent, EndLoc, EndTangent, true);
+		}
+	}
+}
+
+void ACameraRig_Rail::UpdateRailComponents()
+{
+	if (RailSplineComponent && RailCameraMount)
+	{
+		FVector const MountPos = RailSplineComponent->GetLocationAtDistanceAlongSpline(CurrentPositionOnRail, ESplineCoordinateSpace::World);
+		RailCameraMount->SetWorldLocation(MountPos);
+	}
+
+#if WITH_EDITOR
+	if (GIsEditor)
+	{
+		// set up preview mesh to match #todo
+		UpdatePreviewMeshes();
+	}
+#endif
 }
 
+#if WITH_EDITOR
 USceneComponent* ACameraRig_Rail::GetDefaultAttachComponent() const
 {
-	return CraneRailMount;
+	return RailCameraMount;
+}
+
+void ACameraRig_Rail::PostEditChangeProperty(FPropertyChangedEvent& PropertyChangedEvent)
+{
+	Super::PostEditChangeProperty(PropertyChangedEvent);
+	UpdateRailComponents();
+
+	if (RailSplineComponent)
+	{
+		// clamp position to spline limits
+		float const SplineLen = RailSplineComponent->GetSplineLength();
+		CurrentPositionOnRail = FMath::Clamp(CurrentPositionOnRail, 0.f, SplineLen);
+	}
 }
 #endif
 
+void ACameraRig_Rail::Tick(float DeltaTime)
+{
+	Super::Tick(DeltaTime);
 
+	// feed exposed API into underlying components
+	UpdateRailComponents();
+}
+
+bool ACameraRig_Rail::ShouldTickIfViewportsOnly() const
+{
+	return true;
+}
 
-#undef LOCTEXT_NAMESPACE
\ No newline at end of file
+#undef LOCTEXT_NAMESPACE
diff --git a/Engine/Source/Runtime/CinematicCamera/Private/CineCameraActor.cpp b/Engine/Source/Runtime/CinematicCamera/Private/CineCameraActor.cpp
--- a/Engine/Source/Runtime/CinematicCamera/Private/CineCameraActor.cpp
+++ b/Engine/Source/Runtime/CinematicCamera/Private/CineCameraActor.cpp
@@ -70,7 +70,10 @@ void ACineCameraActor::Tick(float DeltaTime)
 		// more complex component hierarchies will require different handling here
 		FVector const LookatLoc = GetLookatLocation();
 		FVector const ToLookat = LookatLoc - GetActorLocation();
-		FRotator const FinalRot = FMath::RInterpTo(LookatTrackingSettings.LastLookatTrackingRotation, ToLookat.Rotation(), DeltaTime, LookatTrackingSettings.LookAtTrackingInterpSpeed);
+		FRotator const FinalRot = 
+			bResetInterplation
+			? ToLookat.Rotation()
+			: FMath::RInterpTo(LookatTrackingSettings.LastLookatTrackingRotation, ToLookat.Rotation(), DeltaTime, LookatTrackingSettings.LookAtTrackingInterpSpeed);
 		SetActorRotation(FinalRot);
 
 		// we store this ourselves in case other systems try to change our rotation, and end fighting the interpolation
@@ -87,6 +90,16 @@ void ACineCameraActor::Tick(float DeltaTime)
 		// no tracking, no ticking
 		SetActorTickEnabled(false);
 	}
+
+	bResetInterplation = false;
+}
+
+
+void ACineCameraActor::NotifyCameraCut()
+{
+	Super::NotifyCameraCut();
+	
+	bResetInterplation = true;
 }
 
 
diff --git a/Engine/Source/Runtime/CinematicCamera/Private/CineCameraComponent.cpp b/Engine/Source/Runtime/CinematicCamera/Private/CineCameraComponent.cpp
--- a/Engine/Source/Runtime/CinematicCamera/Private/CineCameraComponent.cpp
+++ b/Engine/Source/Runtime/CinematicCamera/Private/CineCameraComponent.cpp
@@ -12,6 +12,7 @@
 UCineCameraComponent::UCineCameraComponent()
 {
 	// Super 35mm 4 Perf
+	// These will be overridden if valid default presets are specified in ini
 	FilmbackSettings.SensorWidth = 24.89f;
 	FilmbackSettings.SensorHeight = 18.67;
 	LensSettings.MinFocalLength = 50.f;
@@ -23,6 +24,57 @@ UCineCameraComponent::UCineCameraComponent()
 
 	bConstrainAspectRatio = true;
 
+	RecalcDerivedData();
+
+#if WITH_EDITORONLY_DATA
+	if (!IsRunningCommandlet())
+	{
+		// overrides CameraComponent's camera mesh
+		static ConstructorHelpers::FObjectFinder<UStaticMesh> EditorCameraMesh(TEXT("/Engine/EditorMeshes/Camera/SM_CineCam.SM_CineCam"));
+		CameraMesh = EditorCameraMesh.Object;
+
+		static ConstructorHelpers::FObjectFinder<UStaticMesh> PlaneMesh(TEXT("/Engine/ArtTools/RenderToTexture/Meshes/S_1_Unit_Plane.S_1_Unit_Plane"));
+		DebugFocusPlaneMesh = PlaneMesh.Object;
+
+		static ConstructorHelpers::FObjectFinder<UMaterial> PlaneMat(TEXT("/Engine/EngineDebugMaterials/M_SimpleTranslucent.M_SimpleTranslucent"));
+		DebugFocusPlaneMaterial = PlaneMat.Object;
+	}
+#endif
+}
+
+
+void UCineCameraComponent::PostInitProperties()
+{
+	Super::PostInitProperties();
+
+	// default filmback
+	FNamedFilmbackPreset* const DefaultFilmbackPreset = FilmbackPresets.FindByPredicate(
+		[&](FNamedFilmbackPreset const& Preset)
+		{
+			return (Preset.Name == DefaultFilmbackPresetName);
+		}
+	);
+	if (DefaultFilmbackPreset)
+	{
+		FilmbackSettings = DefaultFilmbackPreset->FilmbackSettings;
+	}
+
+	// default lens
+	FNamedLensPreset* const DefaultLensPreset = LensPresets.FindByPredicate(
+		[&](FNamedLensPreset const& Preset)
+		{
+			return (Preset.Name == DefaultLensPresetName);
+		}
+	);
+	if (DefaultLensPreset)
+	{
+		LensSettings = DefaultLensPreset->LensSettings;
+	}
+
+	// other lens defaults
+	CurrentAperture = DefaultLensFStop;
+	CurrentFocalLength = DefaultLensFocalLength;
+
 	RecalcDerivedData();
 }
 
@@ -37,8 +89,34 @@ void UCineCameraComponent::PostEditChangeProperty(FPropertyChangedEvent& Propert
 {
 	RecalcDerivedData();
 
+	// handle debug focus plane
+	if (FocusSettings.bDrawDebugFocusPlane && (DebugFocusPlaneComponent == nullptr))
+	{
+		CreateDebugFocusPlane();
+	}
+	else if ((FocusSettings.bDrawDebugFocusPlane == false) && (DebugFocusPlaneComponent != nullptr))
+	{
+		DestroyDebugFocusPlane();
+	}
+
+	// set focus plane color in case that's what changed
+	if (DebugFocusPlaneMID)
+	{
+		DebugFocusPlaneMID->SetVectorParameterValue(FName(TEXT("Color")), FocusSettings.DebugFocusPlaneColor.ReinterpretAsLinear());
+	}
+
 	Super::PostEditChangeProperty(PropertyChangedEvent);
 }
+
+void UCineCameraComponent::ResetProxyMeshTransform()
+{
+	if (ProxyMeshComponent)
+	{
+		// CineCam mesh is offset 90deg yaw
+		ProxyMeshComponent->SetRelativeRotation(FRotator(0.f, 90.f, 0.f));
+		ProxyMeshComponent->SetRelativeLocation(FVector(-46.f, 0, -24.f));
+	}
+}
 #endif	// WITH_EDITORONLY_DATA
 
 
@@ -87,8 +165,11 @@ void UCineCameraComponent::RecalcDerivedData()
 	FieldOfView = GetHorizontalFieldOfView();
 	FilmbackSettings.SensorAspectRatio = (FilmbackSettings.SensorHeight > 0.f) ? (FilmbackSettings.SensorWidth / FilmbackSettings.SensorHeight) : 0.f;
 	AspectRatio = FilmbackSettings.SensorAspectRatio;
-}
 
+#if WITH_EDITORONLY_DATA
+	CurrentHorizontalFOV = FieldOfView;			// informational variable only, for editor users
+#endif
+}
 
 static const FName NAME_GetDesiredFocusDistance(TEXT("GetDesiredFocusDistance"));
 static const float SpotFocusTraceDist = 1000000.f;
@@ -197,8 +278,6 @@ float UCineCameraComponent::GetDesiredFocusDistance(FMinimalViewInfo& DesiredVie
 	return DesiredFocusDistance;
 }
 
-static const FColor DebugFocusPlaneColor(64, 32, 128, 64);		// purple
-
 void UCineCameraComponent::GetCameraView(float DeltaTime, FMinimalViewInfo& DesiredView)
 {
 	RecalcDerivedData();
@@ -207,43 +286,133 @@ void UCineCameraComponent::GetCameraView(float DeltaTime, FMinimalViewInfo& Desi
 
 	UpdateCameraLens(DeltaTime, DesiredView);
 
-	if (FocusSettings.bDrawDebugFocusPlane)
+#if WITH_EDITORONLY_DATA
+	if (FocusSettings.bDrawDebugFocusPlane && DebugFocusPlaneMesh)
 	{
 		FVector const CamDir = DesiredView.Rotation.Vector();
 		FVector const FocusPoint = DesiredView.Location + CamDir * CurrentFocusDistance;
+		DebugFocusPlaneComponent->SetWorldLocation(FocusPoint);
+	}
+#endif
+}
 
-		FPlane P(FocusPoint, CamDir);
 
-		::DrawDebugSolidPlane(GetWorld(), P, DesiredView.Location, 10000.f, DebugFocusPlaneColor, false, -1, SDPG_World);
+void UCineCameraComponent::UpdateCameraLens(float DeltaTime, FMinimalViewInfo& DesiredView)
+{
+	if (FocusSettings.FocusMethod == ECameraFocusMethod::None)
+	{
+		DesiredView.PostProcessSettings.bOverride_DepthOfFieldMethod = false;
+		DesiredView.PostProcessSettings.bOverride_DepthOfFieldFstop = false;
+		DesiredView.PostProcessSettings.bOverride_DepthOfFieldFocalDistance = false;
+		DesiredView.PostProcessSettings.bOverride_DepthOfFieldSensorWidth = false;
 	}
+	else
+	{
+		// Update focus/DoF
+		DesiredView.PostProcessBlendWeight = 1.f;
+		DesiredView.PostProcessSettings.bOverride_DepthOfFieldMethod = true;
+		DesiredView.PostProcessSettings.DepthOfFieldMethod = DOFM_CircleDOF;
+
+		DesiredView.PostProcessSettings.bOverride_DepthOfFieldFstop = true;
+		DesiredView.PostProcessSettings.DepthOfFieldFstop = CurrentAperture;
+
+		CurrentFocusDistance = GetDesiredFocusDistance(DesiredView);
+
+		// clamp to min focus distance
+		float const MinFocusDistInWorldUnits = LensSettings.MinimumFocusDistance * (GetWorldToMetersScale() / 1000.f);	// convert mm to uu
+		CurrentFocusDistance = FMath::Max(CurrentFocusDistance, MinFocusDistInWorldUnits);
+
+		// smoothing, if desired
+		if (FocusSettings.bSmoothFocusChanges)
+		{
+			if (bResetInterpolation == false)
+			{
+				CurrentFocusDistance = FMath::FInterpTo(LastFocusDistance, CurrentFocusDistance, DeltaTime, FocusSettings.FocusSmoothingInterpSpeed);
+			}
+		}
+		LastFocusDistance = CurrentFocusDistance;
+
+		DesiredView.PostProcessSettings.bOverride_DepthOfFieldFocalDistance = true;
+		DesiredView.PostProcessSettings.DepthOfFieldFocalDistance = CurrentFocusDistance;
+
+		DesiredView.PostProcessSettings.bOverride_DepthOfFieldSensorWidth = true;
+		DesiredView.PostProcessSettings.DepthOfFieldSensorWidth = FilmbackSettings.SensorWidth;
+	}
+
+	bResetInterpolation = false;
 }
 
+void UCineCameraComponent::NotifyCameraCut()
+{
+	Super::NotifyCameraCut();
 
-void UCineCameraComponent::UpdateCameraLens(float DeltaTime, FMinimalViewInfo& DesiredView)
+	// reset any interpolations
+	bResetInterpolation = true;
+}
+
+#if WITH_EDITORONLY_DATA
+void UCineCameraComponent::CreateDebugFocusPlane()
+{
+	if (AActor* const MyOwner = GetOwner())
+	{
+		if (DebugFocusPlaneComponent == nullptr)
+		{
+			DebugFocusPlaneComponent = NewObject<UStaticMeshComponent>(MyOwner, NAME_None, RF_Transactional | RF_TextExportTransient);
+			DebugFocusPlaneComponent->AttachTo(this);
+			DebugFocusPlaneComponent->AlwaysLoadOnClient = false;
+			DebugFocusPlaneComponent->AlwaysLoadOnServer = false;
+			DebugFocusPlaneComponent->StaticMesh = DebugFocusPlaneMesh;
+			DebugFocusPlaneComponent->SetCollisionProfileName(UCollisionProfile::NoCollision_ProfileName);
+			DebugFocusPlaneComponent->bHiddenInGame = false;
+			DebugFocusPlaneComponent->CastShadow = false;
+			DebugFocusPlaneComponent->PostPhysicsComponentTick.bCanEverTick = false;
+			DebugFocusPlaneComponent->CreationMethod = CreationMethod;
+
+			DebugFocusPlaneComponent->RelativeScale3D = FVector(10000.f, 10000.f, 1.f);
+			DebugFocusPlaneComponent->RelativeRotation = FRotator(90.f, 0.f, 0.f);
+
+			DebugFocusPlaneComponent->RegisterComponentWithWorld(GetWorld());
+
+			DebugFocusPlaneMID = DebugFocusPlaneComponent->CreateAndSetMaterialInstanceDynamicFromMaterial(0, DebugFocusPlaneMaterial);
+			if (DebugFocusPlaneMID)
+			{
+				DebugFocusPlaneMID->SetVectorParameterValue(FName(TEXT("Color")), FocusSettings.DebugFocusPlaneColor.ReinterpretAsLinear());
+			}
+		}
+	}
+}
+
+void UCineCameraComponent::DestroyDebugFocusPlane()
 {
-	// Update focus/DoF
-	DesiredView.PostProcessBlendWeight = 1.f;
-	DesiredView.PostProcessSettings.bOverride_DepthOfFieldMethod = true;
-	DesiredView.PostProcessSettings.DepthOfFieldMethod = DOFM_CircleDOF;
+	if (DebugFocusPlaneComponent)
+	{
+		DebugFocusPlaneComponent->SetVisibility(false);
+		DebugFocusPlaneComponent = nullptr;
 
-	DesiredView.PostProcessSettings.bOverride_DepthOfFieldFstop = true;
-	DesiredView.PostProcessSettings.DepthOfFieldFstop = CurrentAperture;
+		DebugFocusPlaneMID = nullptr;
+	}
+}
+#endif
 
-	CurrentFocusDistance = GetDesiredFocusDistance(DesiredView);
+void UCineCameraComponent::OnRegister()
+{
+	Super::OnRegister();
 
-	// clamp to min focus distance
-	float const MinFocusDistInWorldUnits = LensSettings.MinimumFocusDistance * (GetWorldToMetersScale() / 1000.f);	// convert mm to uu
-	CurrentFocusDistance = FMath::Max(CurrentFocusDistance, MinFocusDistInWorldUnits);
+#if WITH_EDITORONLY_DATA
+	ResetProxyMeshTransform();
+#endif
+}
 
-	// smoothing, if desired
- 	if (FocusSettings.bSmoothFocusChanges)
- 	{
-		CurrentFocusDistance = FMath::FInterpTo(LastFocusDistance, CurrentFocusDistance, DeltaTime, FocusSettings.FocusSmoothingInterpSpeed);
- 	}
-	LastFocusDistance = CurrentFocusDistance;
+#if WITH_EDITOR
+void UCineCameraComponent::OnComponentDestroyed(bool bDestroyingHierarchy)
+{
+	Super::OnComponentDestroyed(bDestroyingHierarchy);
 
-	DesiredView.PostProcessSettings.bOverride_DepthOfFieldFocalDistance = true;
-	DesiredView.PostProcessSettings.DepthOfFieldFocalDistance = CurrentFocusDistance;
+	if (DebugFocusPlaneComponent)
+	{
+		DebugFocusPlaneComponent->DestroyComponent();
+	}
 }
+#endif
 
 #undef LOCTEXT_NAMESPACE
diff --git a/Engine/Source/Runtime/CinematicCamera/Public/CameraRig_Crane.h b/Engine/Source/Runtime/CinematicCamera/Public/CameraRig_Crane.h
--- a/Engine/Source/Runtime/CinematicCamera/Public/CameraRig_Crane.h
+++ b/Engine/Source/Runtime/CinematicCamera/Public/CameraRig_Crane.h
@@ -38,12 +38,11 @@ public:
 	virtual class USceneComponent* GetDefaultAttachComponent() const override;
 #endif
 	
-protected:
+private:
+
 	void UpdatePreviewMeshes();
 	void UpdateCraneComponents();
 
-private:
-
 	/** Root component to give the whole actor a transform. */
 	UPROPERTY(EditDefaultsOnly, Category = "Crane Components")
 	USceneComponent* TransformComponent;
@@ -63,4 +62,9 @@ private:
 	UPROPERTY()
 	UStaticMeshComponent* PreviewMesh_CraneBase;
 
+	UPROPERTY()
+	UStaticMeshComponent* PreviewMesh_CraneMount;
+
+	UPROPERTY()
+	UStaticMeshComponent* PreviewMesh_CraneCounterWeight;
 };
diff --git a/Engine/Source/Runtime/CinematicCamera/Public/CameraRig_Rail.h b/Engine/Source/Runtime/CinematicCamera/Public/CameraRig_Rail.h
--- a/Engine/Source/Runtime/CinematicCamera/Public/CameraRig_Rail.h
+++ b/Engine/Source/Runtime/CinematicCamera/Public/CameraRig_Rail.h
@@ -5,6 +5,8 @@
 #include "CameraRig_Rail.generated.h"
 
 class USceneComponent;
+class USplineComponent;
+class USplineMeshComponent;
 
 /** 
  * 
@@ -17,21 +19,43 @@ class CINEMATICCAMERA_API ACameraRig_Rail : public AActor
 public:
 	// ctor
 	ACameraRig_Rail(const FObjectInitializer& ObjectInitialier);
-
-	UPROPERTY(Interp, EditAnywhere, BlueprintReadWrite, Category = "Crane Controls")
-	float DistanceAlongRail;
 	
+	virtual void Tick(float DeltaTime) override;
+	virtual bool ShouldTickIfViewportsOnly() const override;
+
+	UPROPERTY(Interp, EditAnywhere, BlueprintReadWrite, Category = "Rail Controls")
+	float CurrentPositionOnRail;
+
 #if WITH_EDITOR
-	virtual void PostEditChangeProperty(FPropertyChangedEvent& PropertyChangedEvent) override;
 	virtual class USceneComponent* GetDefaultAttachComponent() const override;
+	virtual void PostEditChangeProperty(FPropertyChangedEvent& PropertyChangedEvent) override;
 #endif
 	
 private:
+	/** Makes sure all components are arranged properly. Call when something changes that might affect components. */
+	void UpdateRailComponents();
+	void UpdatePreviewMeshes();
+	USplineMeshComponent* CreateSplinePreviewSegment();
+
+ 	/** Root component to give the whole actor a transform. */
+	UPROPERTY(EditDefaultsOnly, Category = "Rail Components")
+ 	USceneComponent* TransformComponent;
+
+	UPROPERTY(EditDefaultsOnly, Category = "Rail Components")
+	USplineComponent* RailSplineComponent;
 
-// 	/** Root component to give the whole actor a transform. */
-// 	UPROPERTY()
-// 	USceneComponent* TransformComponent;
+	UPROPERTY(EditDefaultsOnly, Category = "Rail Components")
+	USceneComponent* RailCameraMount;
+
+	UPROPERTY()
+	USplineMeshComponent* PreviewMesh_Rail;
+
+	UPROPERTY()
+	TArray<USplineMeshComponent*> PreviewRailMeshSegments;
+
+	UPROPERTY()
+	UStaticMesh* PreviewRailStaticMesh;
 
 	UPROPERTY()
-	USceneComponent* CraneRailMount;
+	UStaticMeshComponent* PreviewMesh_Mount;
 };
diff --git a/Engine/Source/Runtime/CinematicCamera/Public/CineCameraActor.h b/Engine/Source/Runtime/CinematicCamera/Public/CineCameraActor.h
--- a/Engine/Source/Runtime/CinematicCamera/Public/CineCameraActor.h
+++ b/Engine/Source/Runtime/CinematicCamera/Public/CineCameraActor.h
@@ -23,7 +23,7 @@ struct FCameraLookatTrackingSettings
 	/** Last known lookat tracking rotation (used during interpolation) */
 	FRotator LastLookatTrackingRotation;
 
-	UPROPERTY(/*Interp, */EditAnywhere, BlueprintReadWrite, Category = "LookAt")
+	UPROPERTY(Interp, EditAnywhere, BlueprintReadWrite, Category = "LookAt")
 	AActor* ActorToTrack;
 
 	/** Offset from actor position to look at. Relative to actor if tracking an actor, relative to world otherwise. */
@@ -47,7 +47,7 @@ public:
 	virtual bool ShouldTickIfViewportsOnly() const override;
 	virtual void PostInitializeComponents() override;
 
-	UPROPERTY(EditAnywhere, BlueprintReadWrite, Category = "LookAt")
+	UPROPERTY(EditAnywhere, BlueprintReadWrite, Category = "Current Camera Settings")
 	FCameraLookatTrackingSettings LookatTrackingSettings;
 
 #if WITH_EDITOR
@@ -55,6 +55,11 @@ public:
 #endif
 
 protected:
+
+	/** Set to true to skip any interpolations on the next update. Resets to false automatically. */
+	uint8 bResetInterplation : 1;
+
 	FVector GetLookatLocation() const;
 
+	virtual void NotifyCameraCut() override;
 };
diff --git a/Engine/Source/Runtime/CinematicCamera/Public/CineCameraComponent.h b/Engine/Source/Runtime/CinematicCamera/Public/CineCameraComponent.h
--- a/Engine/Source/Runtime/CinematicCamera/Public/CineCameraComponent.h
+++ b/Engine/Source/Runtime/CinematicCamera/Public/CineCameraComponent.h
@@ -4,11 +4,13 @@
 #include "Camera/CameraComponent.h"
 #include "CineCameraComponent.generated.h"
 
+/** #note, this struct has a details customization in CameraFilmbackSettingsCustomization.cpp/h */
 USTRUCT()
 struct FCameraFilmbackSettings
 {
 	GENERATED_USTRUCT_BODY()
 
+	/** Horizontal size of filmback or digital sensor, in mm. */
 	UPROPERTY(EditAnywhere, BlueprintReadWrite, Category = "Filmback", meta = (ClampMin = "0.001", ForceUnits = mm))
 	float SensorWidth;
 
@@ -16,6 +18,7 @@ struct FCameraFilmbackSettings
 	UPROPERTY(EditAnywhere, BlueprintReadWrite, Category = "Filmback", meta = (ClampMin = "0.001", ForceUnits = mm))
 	float SensorHeight;
 
+	/** Read-only. Computed from Sensor dimensions. */
 	UPROPERTY(VisibleAnywhere, BlueprintReadOnly, Category = "Filmback")
 	float SensorAspectRatio;
 };
@@ -32,6 +35,7 @@ struct FNamedFilmbackPreset
 	FCameraFilmbackSettings FilmbackSettings;
 };
 
+/** #note, this struct has a details customization in CameraLensSettingsCustomization.cpp/h */
 USTRUCT()
 struct FCameraLensSettings
 {
@@ -121,9 +125,12 @@ struct FCameraTrackingFocusSettings
 // 	FName BoneOrSocketName;
 
 	/** Offset from actor position to look at. Relative to actor if tracking an actor, relative to world otherwise. */
-	UPROPERTY(Interp, EditAnywhere, BlueprintReadWrite, Category = "LookAt")
+	UPROPERTY(Interp, EditAnywhere, BlueprintReadWrite, Category = "Tracking Focus")
 	FVector RelativeOffset;
 
+	UPROPERTY(Transient, EditAnywhere, BlueprintReadWrite, Category = "Tracking Focus")
+	uint8 bDrawDebugTrackingFocusPoint : 1;
+
 	FCameraTrackingFocusSettings()
 		: ActorToTrack(nullptr)
 	{}
@@ -135,29 +142,32 @@ struct FCameraFocusSettings
 {
 	GENERATED_USTRUCT_BODY()
 
-	UPROPERTY(EditAnywhere, BlueprintReadWrite, Category = "Focus Settings")
+	UPROPERTY(EditAnywhere, BlueprintReadWrite, Category = "Focus Method")
 	TEnumAsByte<ECameraFocusMethod> FocusMethod;
 	
-	/** True to draw a translucent plane at the current focus depth, for easy tweaking when not looking through the camera. */
-	UPROPERTY(Transient, EditAnywhere, BlueprintReadWrite, Category = "Focus Settings")
-	uint8 bDrawDebugFocusPlane : 1;
-
-	UPROPERTY(EditAnywhere, BlueprintReadWrite, Category = "Focus Settings")
-	uint8 bSmoothFocusChanges : 1;
-	
-	UPROPERTY(/*EditAnywhere, BlueprintReadWrite, Category = "Focus Settings"*/)
-	uint8 bSimulateFocusBreathing : 1;
-
-	UPROPERTY(Interp, EditAnywhere, BlueprintReadWrite, Category = "Focus Settings")
+	UPROPERTY(Interp, EditAnywhere, BlueprintReadWrite, Category = "Manual Focus Settings")
 	float ManualFocusDistance;
 
-	UPROPERTY(/*EditAnywhere, BlueprintReadWrite, Category = "Focus Settings"*/)
+	UPROPERTY(EditAnywhere, BlueprintReadWrite, Category = "Spot Focus Settings")
 	FCameraSpotFocusSettings SpotFocusSettings;
 
-	UPROPERTY(EditAnywhere, BlueprintReadWrite, Category = "Focus Settings")
+	UPROPERTY(EditAnywhere, BlueprintReadWrite, Category = "Tracking Focus Settings")
 	FCameraTrackingFocusSettings TrackingFocusSettings;
 
+#if WITH_EDITORONLY_DATA
+	/** True to draw a translucent plane at the current focus depth, for easy tweaking when not looking through the camera. */
+	UPROPERTY(Transient, EditAnywhere, Category = "Focus Settings")
+	uint8 bDrawDebugFocusPlane : 1;
+
+	/** For customizing the focus plane color, in case the default doesn't show up well in your scene. */
+	UPROPERTY(EditAnywhere, Category = "Focus Settings", meta = (EditCondition = "bDrawDebugFocusPlane"))
+	FColor DebugFocusPlaneColor;
+#endif 
+
 	UPROPERTY(EditAnywhere, BlueprintReadWrite, Category = "Focus Settings")
+	uint8 bSmoothFocusChanges : 1;
+	
+	UPROPERTY(EditAnywhere, BlueprintReadWrite, Category = "Focus Settings", meta = (EditCondition = "bSmoothFocusChanges"))
 	float FocusSmoothingInterpSpeed;
 
 	/** Additional focus depth offset, used for manually tweaking if yoour chosen focus method is needs adjustment */
@@ -166,9 +176,12 @@ struct FCameraFocusSettings
 
 	FCameraFocusSettings() : 
 		FocusMethod(ECameraFocusMethod::Manual),
-		bSmoothFocusChanges(0),
-		bSimulateFocusBreathing(0),
 		ManualFocusDistance(100000.f),
+#if WITH_EDITORONLY_DATA
+		bDrawDebugFocusPlane(0),
+		DebugFocusPlaneColor(102, 26, 204, 153),		// purple
+#endif
+		bSmoothFocusChanges(0),
 		FocusSmoothingInterpSpeed(8.f),
 		FocusOffset(0.f)
 	{}
@@ -214,7 +227,26 @@ public:
 	UPROPERTY(VisibleAnywhere, BlueprintReadOnly, Category = "Current Camera Settings")
 	float CurrentFocusDistance;
 
-	// for interpto
+#if WITH_EDITORONLY_DATA
+	/** Read-only. Control this value with CurrentFocalLength (and filmback settings). */
+	UPROPERTY(VisibleAnywhere, BlueprintReadOnly, Category = "Current Camera Settings")
+	float CurrentHorizontalFOV;
+
+	UPROPERTY(transient)
+	UStaticMesh* DebugFocusPlaneMesh;
+
+	UPROPERTY(transient)
+	UMaterial* DebugFocusPlaneMaterial;
+
+	// The camera mesh to show visually where the camera is placed
+	UPROPERTY(transient)
+	UStaticMeshComponent* DebugFocusPlaneComponent;
+
+	UPROPERTY(transient)
+	UMaterialInstanceDynamic* DebugFocusPlaneMID;
+#endif
+	
+	/** Most recent calculated focus distance. Used for interpolation. */
 	float LastFocusDistance;
 
 	float GetHorizontalFieldOfView() const;
@@ -225,9 +257,20 @@ public:
 
 protected:
 
+	/** Set to true to skip any interpolations on the next update. Resets to false automatically. */
+	uint8 bResetInterpolation : 1;
+
 	virtual void PostLoad() override;
+	virtual void PostInitProperties() override;
+	virtual void OnRegister() override;
+
+#if WITH_EDITOR
+	virtual void OnComponentDestroyed(bool bDestroyingHierarchy) override;
+#endif
+
 #if WITH_EDITORONLY_DATA
 	virtual void PostEditChangeProperty(FPropertyChangedEvent& PropertyChangedEvent) override;
+	virtual void ResetProxyMeshTransform() override;
 #endif
 
 	UPROPERTY(config)
@@ -236,10 +279,29 @@ protected:
 	UPROPERTY(config)
 	TArray<FNamedLensPreset> LensPresets;
 
+	UPROPERTY(config)
+	FString DefaultFilmbackPresetName;
+
+	UPROPERTY(config)
+	FString DefaultLensPresetName;
+	
+	/** Default focal length (will be constrained by default lens) */
+	UPROPERTY(config)
+	float DefaultLensFocalLength;
+	
+	/** Default aperture (will be constrained by default lens) */
+	UPROPERTY(config)
+	float DefaultLensFStop;
+
 	virtual void UpdateCameraLens(float DeltaTime, FMinimalViewInfo& DesiredView);
 
+	virtual void NotifyCameraCut() override;
+	
 private:
 	void RecalcDerivedData();
 	float GetDesiredFocusDistance(FMinimalViewInfo& DesiredView) const;
 	float GetWorldToMetersScale() const;
+
+	void CreateDebugFocusPlane();
+	void DestroyDebugFocusPlane();
 };
diff --git a/Engine/Source/Runtime/Core/Private/Math/UnrealMath.cpp b/Engine/Source/Runtime/Core/Private/Math/UnrealMath.cpp
--- a/Engine/Source/Runtime/Core/Private/Math/UnrealMath.cpp
+++ b/Engine/Source/Runtime/Core/Private/Math/UnrealMath.cpp
@@ -3106,6 +3106,16 @@ bool FMath::Eval( FString Str, float& OutValue )
 	return bResult;
 }
 
+void FMath::WindRelativeAnglesDegrees(float InAngle0, float& InOutAngle1)
+{
+	const float Diff = InAngle0 - InOutAngle1;
+	const float AbsDiff = Abs(Diff);
+	if(AbsDiff > 180.0f)
+	{
+		InOutAngle1 += 360.0f * Sign(Diff) * FloorToFloat((AbsDiff / 360.0f) + 0.5f);
+	}
+}
+
 float FMath::FixedTurn(float InCurrent, float InDesired, float InDeltaRate)
 {
 	if (InDeltaRate == 0.f)
diff --git a/Engine/Source/Runtime/Core/Private/Misc/FeedbackContext.cpp b/Engine/Source/Runtime/Core/Private/Misc/FeedbackContext.cpp
--- a/Engine/Source/Runtime/Core/Private/Misc/FeedbackContext.cpp
+++ b/Engine/Source/Runtime/Core/Private/Misc/FeedbackContext.cpp
@@ -98,9 +98,14 @@ void FFeedbackContext::EndSlowTask()
 }
 /**** End legacy API ****/
 
+bool FFeedbackContext::IsPlayingInEditor() const
+{
+	return GIsPlayInEditorWorld;
+}
+
 void FSlowTask::MakeDialog(bool bShowCancelButton, bool bAllowInPIE)
 {
-	const bool bIsDisabledByPIE = GIsPlayInEditorWorld && !bAllowInPIE;
+	const bool bIsDisabledByPIE = Context.IsPlayingInEditor() && !bAllowInPIE;
 	const bool bIsDialogAllowed = bEnabled && !GIsSilent && !bIsDisabledByPIE && !IsRunningCommandlet() && IsInGameThread();
 	if (!GIsSlowTask && bIsDialogAllowed)
 	{
diff --git a/Engine/Source/Runtime/Core/Public/Math/UnrealMathUtility.h b/Engine/Source/Runtime/Core/Public/Math/UnrealMathUtility.h
--- a/Engine/Source/Runtime/Core/Public/Math/UnrealMathUtility.h
+++ b/Engine/Source/Runtime/Core/Public/Math/UnrealMathUtility.h
@@ -473,6 +473,14 @@ struct FMath : public FPlatformMath
 		return A;
 	}
 
+	/** 
+	 * Given two angles in degrees, 'wind' the rotation in Angle1 so that it avoids >180 degree flips.
+	 * Good for winding rotations previously expressed as quaternions into a euler-angle representation.
+	 * @param	Angle0	The first angle that we wind relative to.
+	 * @param	Angle1	The second angle that we may wind relative to the first.
+	 */
+	static CORE_API void WindRelativeAnglesDegrees(float InAngle0, float& InOutAngle1);
+
 	/** Returns a new rotation component value
 	 *
 	 * @param InCurrent is the current rotation value
diff --git a/Engine/Source/Runtime/Core/Public/Misc/FeedbackContext.h b/Engine/Source/Runtime/Core/Public/Misc/FeedbackContext.h
--- a/Engine/Source/Runtime/Core/Public/Misc/FeedbackContext.h
+++ b/Engine/Source/Runtime/Core/Public/Misc/FeedbackContext.h
@@ -67,6 +67,9 @@ protected:
 	 */
 	virtual void ProgressReported( const float TotalProgressInterp, FText DisplayMessage ) {}
 
+	/** Called to check whether we are playing in editor when starting a slow task */
+	virtual bool IsPlayingInEditor() const;
+
 public:
 	virtual FContextSupplier* GetContext() const { return NULL; }
 	virtual void SetContext( FContextSupplier* InSupplier ) {}
diff --git a/Engine/Source/Runtime/CoreUObject/Public/UObject/ObjectBase.h b/Engine/Source/Runtime/CoreUObject/Public/UObject/ObjectBase.h
--- a/Engine/Source/Runtime/CoreUObject/Public/UObject/ObjectBase.h
+++ b/Engine/Source/Runtime/CoreUObject/Public/UObject/ObjectBase.h
@@ -1114,6 +1114,9 @@ namespace UM
 		/// [PropertyMetadta] Used by FDirectoryPath properties. Indicates that the directory dialog will output a path relative to the game content directory when setting the property.
 		RelativeToGameContentDir,
 
+		/// [PropertyMetadta] Used by FDirectoryPath properties. Indicates that the path will be picked using the Slate-style directory picker inside the game Content dir.
+		ContentDir,
+
 		// [PropertyMetadata] Used by struct properties. Indicates that the inner properties will not be shown inside an expandable struct, but promoted up a level.
 		ShowOnlyInnerProperties,
 
diff --git a/Engine/Source/Runtime/Engine/Classes/Animation/AnimInstance.h b/Engine/Source/Runtime/Engine/Classes/Animation/AnimInstance.h
--- a/Engine/Source/Runtime/Engine/Classes/Animation/AnimInstance.h
+++ b/Engine/Source/Runtime/Engine/Classes/Animation/AnimInstance.h
@@ -601,10 +601,14 @@ public:
 	UFUNCTION(BlueprintCallable, Category = "Animation")
 	void Montage_Stop(float InBlendOutTime, UAnimMontage * Montage = NULL);
 
-	/** Pauses the animation montage. If reference is NULL, it will stop ALL active montages. */
+	/** Pauses the animation montage. If reference is NULL, it will pause ALL active montages. */
 	UFUNCTION(BlueprintCallable, Category = "Animation")
 	void Montage_Pause(UAnimMontage * Montage = NULL);
 
+	/** Resumes a paused animation montage. If reference is NULL, it will resume ALL active montages. */
+	UFUNCTION(BlueprintCallable, Category = "Animation")
+	void Montage_Resume(UAnimMontage* Montage);
+
 	/** Makes a montage jump to a named section. If Montage reference is NULL, it will do that to all active montages. */
 	UFUNCTION(BlueprintCallable, Category="Animation")
 	void Montage_JumpToSection(FName SectionName, UAnimMontage * Montage = NULL);
diff --git a/Engine/Source/Runtime/Engine/Classes/Animation/AnimMontage.h b/Engine/Source/Runtime/Engine/Classes/Animation/AnimMontage.h
--- a/Engine/Source/Runtime/Engine/Classes/Animation/AnimMontage.h
+++ b/Engine/Source/Runtime/Engine/Classes/Animation/AnimMontage.h
@@ -325,6 +325,7 @@ public:
 
 	bool IsValid() const { return (Montage!=NULL); }
 	bool IsPlaying() const { return IsValid() && bPlaying; }
+	void SetPlaying(bool bInPlaying) { bPlaying = bInPlaying; }
 	bool IsStopped() const { return Blend.GetDesiredValue() == 0.f; }
 
 	/** Returns true if this montage is active (valid and not blending out) */
diff --git a/Engine/Source/Runtime/Engine/Classes/Camera/CameraActor.h b/Engine/Source/Runtime/Engine/Classes/Camera/CameraActor.h
--- a/Engine/Source/Runtime/Engine/Classes/Camera/CameraActor.h
+++ b/Engine/Source/Runtime/Engine/Classes/Camera/CameraActor.h
@@ -74,4 +74,11 @@ public:
 
 	/** Returns CameraComponent subobject **/
 	class UCameraComponent* GetCameraComponent() const;
+
+	/** 
+	 * Called to notify that this camera was cut to, so it can update things like interpolation if necessary.
+	 * Typically called by the camera component.
+	 */
+	virtual void NotifyCameraCut() {};
+
 };
diff --git a/Engine/Source/Runtime/Engine/Classes/Camera/CameraAnimInst.h b/Engine/Source/Runtime/Engine/Classes/Camera/CameraAnimInst.h
--- a/Engine/Source/Runtime/Engine/Classes/Camera/CameraAnimInst.h
+++ b/Engine/Source/Runtime/Engine/Classes/Camera/CameraAnimInst.h
@@ -33,6 +33,10 @@ public:
 	uint32 bFinished : 1;
 
 protected:
+
+	/** If true, this anim inst will automatically stop itself when it finishes, otherwise, it will wait for an explicit Stop() call. */
+	uint32 bStopAutomatically : 1;
+
 	/** True if the animation should loop, false otherwise. */
 	uint32 bLooping:1;
 
@@ -116,6 +120,12 @@ public:
 	/** advances the animation by the specified time - updates any modified interp properties, moves the group actor, etc */
 	void AdvanceAnim(float DeltaTime, bool bJump);
 	
+	/** Jumps he camera anim to the given (unscaled) time. */
+	void SetCurrentTime(float NewTime);
+
+	/** Returns the current playback time. */
+	float GetCurrentTime() const { return CurTime; };
+
 	/** Stops this instance playing whatever animation it is playing. */
 	UFUNCTION(BlueprintCallable, Category = CameraAnimInst)
 	void Stop(bool bImmediate = false);
@@ -134,6 +144,15 @@ public:
 	UFUNCTION(BlueprintCallable, Category = CameraAnimInst)
 	void SetScale(float NewDuration);
 
+	/** Sets the camera actor that will be modified when this anim is played. (Usually a temporary actor whose data is subsequently harvested). */
+	void SetCameraActor(class AActor* Actor);
+
+	/** Takes the given view and applies the camera anim transform and fov changes to it. Does not affect PostProcess. */
+	void ApplyToView(FMinimalViewInfo& InOutPOV) const;
+
+	/** Sets whether this anim instance should automatically stop when finished. */
+	void SetStopAutomatically(bool bNewStopAutoMatically) { bStopAutomatically = bNewStopAutoMatically; };
+
 protected:
 	/** Returns InterpGroupInst subobject **/
 	class UInterpGroupInst* GetInterpGroupInst() const;
diff --git a/Engine/Source/Runtime/Engine/Classes/Camera/CameraComponent.h b/Engine/Source/Runtime/Engine/Classes/Camera/CameraComponent.h
--- a/Engine/Source/Runtime/Engine/Classes/Camera/CameraComponent.h
+++ b/Engine/Source/Runtime/Engine/Classes/Camera/CameraComponent.h
@@ -67,6 +67,11 @@ class ENGINE_API UCameraComponent : public USceneComponent
 	UPROPERTY(EditAnywhere, BlueprintReadWrite, Category=CameraSettings)
 	uint32 bUsePawnControlRotation:1;
 
+protected:
+	/** True to enable the additive view offset, for adjusting the view without moving the component. */
+	uint32 bUseAdditiveOffset : 1;
+
+public:
 	// The type of camera
 	UPROPERTY(Interp, EditAnywhere, BlueprintReadWrite, Category=CameraSettings)
 	TEnumAsByte<ECameraProjectionMode::Type> ProjectionMode;
@@ -122,17 +127,46 @@ protected:
 	// The camera mesh to show visually where the camera is placed
 	class UStaticMeshComponent* ProxyMeshComponent;
 
-public:
-
 	virtual void SetCameraMesh(UStaticMesh* Mesh);
 
+	virtual void ResetProxyMeshTransform();
 #endif
+
+	/** An optional extra transform to adjust the final view without moving the component, in the camera's local space */
+	FTransform AdditiveOffset;
+
+	/** An optional extra FOV offset to adjust the final view without modifying the component */
+	float AdditiveFOVOffset;
+
+public:
+
+	/** Applies the given additive offset, preserving any existing offset */
+	void AddAdditiveOffset(FTransform const& Transform, float FOV)
+	{
+		bUseAdditiveOffset = true;
+		AdditiveOffset = AdditiveOffset * Transform;
+		AdditiveFOVOffset += FOV;
+	}
+
+	/** Removes any additive offset. */
+	void ClearAdditiveOffset()
+	{
+		bUseAdditiveOffset = false;
+		AdditiveOffset = FTransform::Identity;
+		AdditiveFOVOffset = 0.f;
+	}
+
+	/** 
+	 * Can be called from external code to notify that this camera was cut to, so it can update 
+	 * things like interpolation if necessary.
+	 */
+	virtual void NotifyCameraCut();
+
 public:
 #if WITH_EDITORONLY_DATA
 	// Refreshes the visual components to match the component state
 	virtual void RefreshVisualRepresentation();
 
-
 	void OverrideFrustumColor(FColor OverrideColor);
 	void RestoreFrustumColor();
 #endif
diff --git a/Engine/Source/Runtime/Engine/Classes/Camera/CameraShake.h b/Engine/Source/Runtime/Engine/Classes/Camera/CameraShake.h
--- a/Engine/Source/Runtime/Engine/Classes/Camera/CameraShake.h
+++ b/Engine/Source/Runtime/Engine/Classes/Camera/CameraShake.h
@@ -48,6 +48,9 @@ struct FFOscillator
 
 	/** Returns the initial value of the oscillator. */
 	static float GetInitialOffset(FFOscillator const& Osc);
+
+	/** Returns the offset at the given time */
+	static float GetOffsetAtTime(FFOscillator const& Osc, float InitialOffset, float Time);
 };
 
 /** Defines FRotator oscillation. */
@@ -208,6 +211,15 @@ protected:
 	/** Current FOV sinusoidal offset. */
 	float FOVSinOffset;
 
+	/** Initial offset (could have been assigned at random). */
+	FVector InitialLocSinOffset;
+
+	/** Initial offset (could have been assigned at random). */
+	FVector InitialRotSinOffset;
+
+	/** Initial offset (could have been assigned at random). */
+	float InitialFOVSinOffset;
+
 	/** Matrix defining the playspace, used when PlaySpace == CAPS_UserDefined */
 	FMatrix UserPlaySpaceMatrix;
 
@@ -259,6 +271,9 @@ public:
 
 	// Returns true if this camera shake will loop forever
 	bool IsLooping() const;
+
+	/** Sets current playback time and applies the shake (both oscillation and cameraanim) to the given POV. */
+	void SetCurrentTimeAndApplyShake(float NewTime, FMinimalViewInfo& POV);
 };
 
 
diff --git a/Engine/Source/Runtime/Engine/Classes/Components/SkeletalMeshComponent.h b/Engine/Source/Runtime/Engine/Classes/Components/SkeletalMeshComponent.h
--- a/Engine/Source/Runtime/Engine/Classes/Components/SkeletalMeshComponent.h
+++ b/Engine/Source/Runtime/Engine/Classes/Components/SkeletalMeshComponent.h
@@ -1463,6 +1463,21 @@ public:
 	FTransform ConvertLocalRootMotionToWorld(const FTransform& InTransform);
 
 	FRootMotionMovementParams ConsumeRootMotion();
+
+private:
+
+#if WITH_EDITOR
+	/** This is required for recording animations, so save for editor only */
+	/** Temporary array of curve arrays that are active on this component - keeps same buffer index as SpaceBases - Please check SkinnedMeshComponent*/
+	FBlendedHeapCurve	CurvesArray[2];
+public: 
+	/** Access Curve Array for reading */
+	const FBlendedHeapCurve& GetAnimationCurves() const { return CurvesArray[CurrentReadSpaceBases]; }
+
+	/** Get Access to the current editable Curve Array - uses same buffer as space bases*/
+	FBlendedHeapCurve& GetEditableAnimationCurves() { return CurvesArray[CurrentEditableSpaceBases]; }
+	const FBlendedHeapCurve& GetEditableAnimationCurves() const { return CurvesArray[CurrentEditableSpaceBases]; }
+#endif 
 };
 
 
diff --git a/Engine/Source/Runtime/Engine/Classes/Curves/CurveBase.h b/Engine/Source/Runtime/Engine/Classes/Curves/CurveBase.h
--- a/Engine/Source/Runtime/Engine/Classes/Curves/CurveBase.h
+++ b/Engine/Source/Runtime/Engine/Classes/Curves/CurveBase.h
@@ -703,7 +703,7 @@ typedef FRichCurveEditInfoTemplate<const FRichCurve*>	FRichCurveEditInfoConst;
 /**
  * Interface you implement if you want the CurveEditor to be able to edit curves on you.
  */
-class FCurveOwnerInterface
+class ENGINE_API FCurveOwnerInterface
 {
 public:
 
@@ -744,6 +744,9 @@ public:
 
 	/** Validates that a previously retrieved curve is still valid for editing. */
 	virtual bool IsValidCurve(FRichCurveEditInfo CurveInfo) = 0;
+
+	/** @return Color for this curve */
+	virtual FLinearColor GetCurveColor(FRichCurveEditInfo CurveInfo) const;
 };
 
 
@@ -840,11 +843,11 @@ public:
 	{ }
 	
 	/** The keyed time */
-	UPROPERTY()
+	UPROPERTY(EditAnywhere, Category=Key)
 	float Time;
 
 	/** The keyed integral value */
-	UPROPERTY()
+	UPROPERTY(EditAnywhere, Category=Key)
 	int32 Value;
 };
 
diff --git a/Engine/Source/Runtime/Engine/Classes/Engine/Engine.h b/Engine/Source/Runtime/Engine/Classes/Engine/Engine.h
--- a/Engine/Source/Runtime/Engine/Classes/Engine/Engine.h
+++ b/Engine/Source/Runtime/Engine/Classes/Engine/Engine.h
@@ -1743,10 +1743,10 @@ public:
 	virtual float GetMaxTickRate(float DeltaTime, bool bAllowFrameRateSmoothing = true) const;
 
 	/** Get max fps. */
-	virtual int32 GetMaxFPS() const;
+	virtual float GetMaxFPS() const;
 
 	/** Set max fps. Overrides console variable. */
-	virtual void SetMaxFPS(const int32 MaxFPS);
+	virtual void SetMaxFPS(const float MaxFPS);
 
 	/** Updates the running average delta time */
 	virtual void UpdateRunningAverageDeltaTime(float DeltaTime, bool bAllowFrameRateSmoothing = true);
diff --git a/Engine/Source/Runtime/Engine/Classes/Engine/Scene.h b/Engine/Source/Runtime/Engine/Classes/Engine/Scene.h
--- a/Engine/Source/Runtime/Engine/Classes/Engine/Scene.h
+++ b/Engine/Source/Runtime/Engine/Classes/Engine/Scene.h
@@ -372,6 +372,9 @@ struct FPostProcessSettings
 	UPROPERTY(EditAnywhere, BlueprintReadWrite, Category=Overrides, meta=(PinHiddenByDefault, InlineEditConditionToggle))
 	uint32 bOverride_DepthOfFieldFstop:1;
 
+	UPROPERTY(EditAnywhere, BlueprintReadWrite, Category=Overrides, meta=(PinHiddenByDefault, InlineEditConditionToggle))
+	uint32 bOverride_DepthOfFieldSensorWidth:1;
+
 	UPROPERTY(EditAnywhere, BlueprintReadWrite, Category=Overrides, meta=(PinHiddenByDefault, InlineEditConditionToggle))
 	uint32 bOverride_DepthOfFieldDepthBlurRadius:1;
 
@@ -848,6 +851,10 @@ struct FPostProcessSettings
 	UPROPERTY(interp, BlueprintReadWrite, Category=DepthOfField, meta=(ClampMin = "1.0", ClampMax = "32.0", editcondition = "bOverride_DepthOfFieldFstop", DisplayName = "Aperture F-stop"))
 	float DepthOfFieldFstop;
 
+	/** Width of the camera sensor to assume, in mm. */
+	UPROPERTY(BlueprintReadWrite, Category=DepthOfField, meta=(ForceUnits=mm, ClampMin = "0.1", UIMin="0.1", UIMax= "1000.0", editcondition = "bOverride_DepthOfFieldSensorWidth", DisplayName = "Sensor Width (mm)"))
+	float DepthOfFieldSensorWidth;
+
 	/** Distance in which the Depth of Field effect should be sharp, in unreal units (cm) */
 	UPROPERTY(interp, BlueprintReadWrite, Category=DepthOfField, meta=(ClampMin = "0.0", UIMin = "1.0", UIMax = "10000.0", editcondition = "bOverride_DepthOfFieldFocalDistance", DisplayName = "Focal Distance"))
 	float DepthOfFieldFocalDistance;
@@ -1131,6 +1138,7 @@ struct FPostProcessSettings
 		ColorGradingIntensity = 1.0f;
 		DepthOfFieldFocalDistance = 1000.0f;
 		DepthOfFieldFstop = 4.0f;
+		DepthOfFieldSensorWidth = 24.576f;			// APS-C
 		DepthOfFieldDepthBlurAmount = 1.0f;
 		DepthOfFieldDepthBlurRadius = 0.0f;
 		DepthOfFieldFocalRegion = 0.0f;
diff --git a/Engine/Source/Runtime/Engine/Classes/Kismet/KismetMathLibrary.h b/Engine/Source/Runtime/Engine/Classes/Kismet/KismetMathLibrary.h
--- a/Engine/Source/Runtime/Engine/Classes/Kismet/KismetMathLibrary.h
+++ b/Engine/Source/Runtime/Engine/Classes/Kismet/KismetMathLibrary.h
@@ -441,6 +441,10 @@ class ENGINE_API UKismetMathLibrary : public UBlueprintFunctionLibrary
 	UFUNCTION(BlueprintPure, meta=(DisplayName = "Get PI", CompactNodeTitle = "PI"), Category="Math|Trig")
 	static float GetPI();
 
+	/* Returns the value of TAU (= 2 * PI) */
+	UFUNCTION(BlueprintPure, meta=(DisplayName = "Get TAU", CompactNodeTitle = "TAU"), Category="Math|Trig")
+	static float GetTAU();
+
 	/* Returns radians value based on the input degrees */
 	UFUNCTION(BlueprintPure, meta=(DisplayName = "Degrees To Radians", CompactNodeTitle = "D2R"), Category="Math|Trig")
 	static float DegreesToRadians(float A);
