Copying //UE4/Dev-Sequencer to //UE4/Dev-Main (Source: //UE4/Dev-Sequencer @ 3946692)

diff --git a/Engine/Source/Editor/SequencerWidgets/Private/STimeRange.cpp b/Engine/Source/Editor/SequencerWidgets/Private/STimeRange.cpp
--- a/Engine/Source/Editor/SequencerWidgets/Private/STimeRange.cpp
+++ b/Engine/Source/Editor/SequencerWidgets/Private/STimeRange.cpp
@@ -1,3 +1,4 @@
+
 // Copyright 1998-2018 Epic Games, Inc. All Rights Reserved.
 
 #include "STimeRange.h"
@@ -7,89 +8,104 @@
 #include "Widgets/Layout/SBox.h"
 #include "Widgets/Input/SSpinBox.h"
 #include "EditorStyleSet.h"
+#include "QualifiedFrameTime.h"
+#include "MovieSceneTimeHelpers.h"
 
 #define LOCTEXT_NAMESPACE "STimeRange"
 
-void STimeRange::Construct( const STimeRange::FArguments& InArgs, TSharedRef<ITimeSliderController> InTimeSliderController, TSharedRef<INumericTypeInterface<float>> NumericTypeInterface )
+void STimeRange::Construct( const STimeRange::FArguments& InArgs, TSharedRef<ITimeSliderController> InTimeSliderController, TSharedRef<INumericTypeInterface<double>> NumericTypeInterface )
 {
 	TimeSliderController = InTimeSliderController;
-	ShowFrameNumbers = InArgs._ShowFrameNumbers;
 
 	TSharedRef<SWidget> WorkingRangeStart = SNullWidget::NullWidget, WorkingRangeEnd = SNullWidget::NullWidget;
 	if (InArgs._ShowWorkingRange)
 	{
-		WorkingRangeStart = SNew(SSpinBox<float>)
+		WorkingRangeStart = SNew(SSpinBox<double>)
 		.Value(this, &STimeRange::WorkingStartTime)
 		.ToolTipText(LOCTEXT("WorkingRangeStart", "Working Range Start"))
-		.OnValueCommitted( this, &STimeRange::OnWorkingStartTimeCommitted )
-		.OnValueChanged( this, &STimeRange::OnWorkingStartTimeChanged )
-		.MinValue(TOptional<float>())
+		.OnValueCommitted(this, &STimeRange::OnWorkingStartTimeCommitted)
+		.OnValueChanged(this, &STimeRange::OnWorkingStartTimeChanged)
+		.MinValue(TOptional<double>())
 		.MaxValue(this, &STimeRange::MaxWorkingStartTime)
 		.Style(&FEditorStyle::Get().GetWidgetStyle<FSpinBoxStyle>("Sequencer.HyperlinkSpinBox"))
 		.TypeInterface(NumericTypeInterface)
-		.ClearKeyboardFocusOnCommit(true);
+		.ClearKeyboardFocusOnCommit(true)
+		.Delta(this, &STimeRange::GetSpinboxDelta)
+		.LinearDeltaSensitivity(25);
 
-		WorkingRangeEnd = SNew(SSpinBox<float>)
+		WorkingRangeEnd = SNew(SSpinBox<double>)
 		.Value(this, &STimeRange::WorkingEndTime)
 		.ToolTipText(LOCTEXT("WorkingRangeEnd", "Working Range End"))
 		.OnValueCommitted( this, &STimeRange::OnWorkingEndTimeCommitted )
 		.OnValueChanged( this, &STimeRange::OnWorkingEndTimeChanged )
 		.MinValue(this, &STimeRange::MinWorkingEndTime)
-		.MaxValue(TOptional<float>())
+		.MaxValue(TOptional<double>())
 		.Style(&FEditorStyle::Get().GetWidgetStyle<FSpinBoxStyle>("Sequencer.HyperlinkSpinBox"))
 		.TypeInterface(NumericTypeInterface)
-		.ClearKeyboardFocusOnCommit(true);
+		.ClearKeyboardFocusOnCommit(true)
+		.Delta(this, &STimeRange::GetSpinboxDelta)
+		.LinearDeltaSensitivity(25);
 	}
 
 	TSharedRef<SWidget> ViewRangeStart = SNullWidget::NullWidget, ViewRangeEnd = SNullWidget::NullWidget;
 	if (InArgs._ShowViewRange)
 	{
-		ViewRangeStart = SNew(SSpinBox<float>)
+		ViewRangeStart = SNew(SSpinBox<double>)
 		.Value(this, &STimeRange::ViewStartTime)
 		.ToolTipText(ViewStartTimeTooltip())
 		.OnValueCommitted( this, &STimeRange::OnViewStartTimeCommitted )
 		.OnValueChanged( this, &STimeRange::OnViewStartTimeChanged )
-		.MinValue(TOptional<float>())
+		.MinValue(TOptional<double>())
 		.MaxValue(this, &STimeRange::MaxViewStartTime)
 		.Style(&FEditorStyle::Get().GetWidgetStyle<FSpinBoxStyle>("Sequencer.HyperlinkSpinBox"))
 		.TypeInterface(NumericTypeInterface)
-		.ClearKeyboardFocusOnCommit(true);
+		.ClearKeyboardFocusOnCommit(true)
+		.Delta(this, &STimeRange::GetSpinboxDelta)
+		.LinearDeltaSensitivity(25);
+
 
-		ViewRangeEnd = SNew(SSpinBox<float>)
+		ViewRangeEnd = SNew(SSpinBox<double>)
 		.Value(this, &STimeRange::ViewEndTime)
 		.ToolTipText(ViewEndTimeTooltip())
 		.OnValueCommitted( this, &STimeRange::OnViewEndTimeCommitted )
 		.OnValueChanged( this, &STimeRange::OnViewEndTimeChanged )
 		.MinValue(this, &STimeRange::MinViewEndTime)
-		.MaxValue(TOptional<float>())
+		.MaxValue(TOptional<double>())
 		.Style(&FEditorStyle::Get().GetWidgetStyle<FSpinBoxStyle>("Sequencer.HyperlinkSpinBox"))
 		.TypeInterface(NumericTypeInterface)
-		.ClearKeyboardFocusOnCommit(true);
+		.ClearKeyboardFocusOnCommit(true)
+		.Delta(this, &STimeRange::GetSpinboxDelta)
+		.LinearDeltaSensitivity(25);
 	}
 
 	TSharedRef<SWidget> PlaybackRangeStart = SNullWidget::NullWidget, PlaybackRangeEnd = SNullWidget::NullWidget;
 	{
-		PlaybackRangeStart = SNew(SSpinBox<float>)
+		PlaybackRangeStart = SNew(SSpinBox<double>)
 		.Value(this, &STimeRange::PlayStartTime)
 		.ToolTipText(PlayStartTimeTooltip())
-		.OnValueCommitted( this, &STimeRange::OnPlayStartTimeCommitted )
-		.OnValueChanged( this, &STimeRange::OnPlayStartTimeChanged )
-		.MinValue(this, &STimeRange::MinPlayStartTime)
-		.MaxValue(this, &STimeRange::MaxPlayStartTime)
+		.OnValueCommitted(this, &STimeRange::OnPlayStartTimeCommitted)
+		.OnValueChanged(this, &STimeRange::OnPlayStartTimeChanged)
+		.MinValue(TOptional<double>())
+		.MaxValue(TOptional<double>())
 		.Style(&FEditorStyle::Get().GetWidgetStyle<FSpinBoxStyle>("Sequencer.HyperlinkSpinBox"))
 		.TypeInterface(NumericTypeInterface)
-		.ClearKeyboardFocusOnCommit(true);
+		.ClearKeyboardFocusOnCommit(true)
+		.Delta(this, &STimeRange::GetSpinboxDelta)
+		.LinearDeltaSensitivity(25);
+
 
-		PlaybackRangeEnd = SNew(SSpinBox<float>)
+		PlaybackRangeEnd = SNew(SSpinBox<double>)
 		.Value(this, &STimeRange::PlayEndTime)
 		.ToolTipText(PlayEndTimeTooltip())
 		.OnValueCommitted( this, &STimeRange::OnPlayEndTimeCommitted )
 		.OnValueChanged( this, &STimeRange::OnPlayEndTimeChanged )
-		.MinValue(this, &STimeRange::MinPlayEndTime)
-		.MaxValue(this, &STimeRange::MaxPlayEndTime)
+		.MinValue(TOptional<double>())
+		.MaxValue(TOptional<double>())
 		.Style(&FEditorStyle::Get().GetWidgetStyle<FSpinBoxStyle>("Sequencer.HyperlinkSpinBox"))
 		.TypeInterface(NumericTypeInterface)
-		.ClearKeyboardFocusOnCommit(true);
+		.ClearKeyboardFocusOnCommit(true)
+		.Delta(this, &STimeRange::GetSpinboxDelta)
+		.LinearDeltaSensitivity(25);
 	}
 
 	this->ChildSlot
@@ -104,7 +120,7 @@ void STimeRange::Construct( const STimeRange::FArguments& InArgs, TSharedRef<ITi
 		[
 			SNew(SBox)
 			.Visibility(InArgs._ShowWorkingRange ? EVisibility::Visible : EVisibility::Collapsed)
-			.WidthOverride(48)
+			.MinDesiredWidth(64)
 			.HAlign(HAlign_Center)
 			[
 				WorkingRangeStart
@@ -117,7 +133,7 @@ void STimeRange::Construct( const STimeRange::FArguments& InArgs, TSharedRef<ITi
 		.Padding(2.f)
 		[
 			SNew(SBox)
-			.WidthOverride(48)
+			.MinDesiredWidth(64)
 			.HAlign(HAlign_Center)
 			.Visibility(InArgs._ShowPlaybackRange ? EVisibility::Visible : EVisibility::Collapsed)
 			[
@@ -138,7 +154,7 @@ void STimeRange::Construct( const STimeRange::FArguments& InArgs, TSharedRef<ITi
 		[
 			SNew(SBox)
 			.Visibility(InArgs._ShowViewRange ? EVisibility::Visible : EVisibility::Collapsed)
-			.WidthOverride(48)
+			.MinDesiredWidth(64)
 			.HAlign(HAlign_Center)
 			[
 				ViewRangeStart
@@ -160,7 +176,7 @@ void STimeRange::Construct( const STimeRange::FArguments& InArgs, TSharedRef<ITi
 		[
 			SNew(SBox)
 			.Visibility(InArgs._ShowViewRange ? EVisibility::Visible : EVisibility::Collapsed)
-			.WidthOverride(48)
+			.MinDesiredWidth(64)
 			.HAlign(HAlign_Center)
 			[
 				ViewRangeEnd
@@ -173,7 +189,7 @@ void STimeRange::Construct( const STimeRange::FArguments& InArgs, TSharedRef<ITi
 		.Padding(2.f)
 		[
 			SNew(SBox)
-			.WidthOverride(48)
+			.MinDesiredWidth(64)
 			.HAlign(HAlign_Center)
 			.Visibility(InArgs._ShowPlaybackRange ? EVisibility::Visible : EVisibility::Collapsed)
 			[
@@ -193,7 +209,7 @@ void STimeRange::Construct( const STimeRange::FArguments& InArgs, TSharedRef<ITi
 		.Padding(2.f)
 		[
 			SNew(SBox)
-			.WidthOverride(48)
+			.MinDesiredWidth(64)
 			.HAlign(HAlign_Center)
 			.Visibility(InArgs._ShowWorkingRange ? EVisibility::Visible : EVisibility::Collapsed)
 			[
@@ -203,234 +219,280 @@ void STimeRange::Construct( const STimeRange::FArguments& InArgs, TSharedRef<ITi
 	];
 }
 
-float STimeRange::WorkingStartTime() const
+double STimeRange::WorkingStartTime() const
+{
+	FFrameRate Rate = TimeSliderController->GetFrameResolution();
+	FFrameTime Time = TimeSliderController->GetClampRange().GetLowerBoundValue() * Rate;
+	return Time.GetFrame().Value;
+}
+
+double STimeRange::WorkingEndTime() const
 {
-	return TimeSliderController.IsValid() ? TimeSliderController.Get()->GetClampRange().GetLowerBoundValue() : 0.f;
+	FFrameRate Rate = TimeSliderController->GetFrameResolution();
+	FFrameTime Time = TimeSliderController->GetClampRange().GetUpperBoundValue() * Rate;
+	return Time.GetFrame().Value;
 }
 
-float STimeRange::WorkingEndTime() const
+double STimeRange::ViewStartTime() const
 {
-	return TimeSliderController.IsValid() ? TimeSliderController.Get()->GetClampRange().GetUpperBoundValue() : 0.f;
+	FFrameRate FrameResolution = TimeSliderController->GetFrameResolution();
+
+	// View range is in seconds so we convert it to frame resolution
+	FFrameTime Time = TimeSliderController->GetViewRange().GetLowerBoundValue() * FrameResolution;
+	return Time.GetFrame().Value;
 }
 
-float STimeRange::ViewStartTime() const
+double STimeRange::ViewEndTime() const
 {
-	return TimeSliderController.IsValid() ? TimeSliderController.Get()->GetViewRange().GetLowerBoundValue() : 0.f;
+	FFrameRate FrameResolution = TimeSliderController->GetFrameResolution();
+
+	// View range is in seconds so we convert it to frame resolution
+	FFrameTime Time = TimeSliderController->GetViewRange().GetUpperBoundValue() * FrameResolution;
+	return Time.GetFrame().Value;
 }
 
-float STimeRange::ViewEndTime() const
+double STimeRange::GetSpinboxDelta() const
 {
-	return TimeSliderController.IsValid() ? TimeSliderController.Get()->GetViewRange().GetUpperBoundValue() : 0.f;
+	return TimeSliderController->GetFrameResolution().AsDecimal() * TimeSliderController->GetPlayRate().AsInterval();
 }
 
-float STimeRange::PlayStartTime() const
+double STimeRange::PlayStartTime() const
 {
-	return TimeSliderController.IsValid() ? TimeSliderController.Get()->GetPlayRange().GetLowerBoundValue() : 0.f;
+	FFrameNumber LowerBound = MovieScene::DiscreteInclusiveLower(TimeSliderController->GetPlayRange());
+	return LowerBound.Value; 
 }
 
-float STimeRange::PlayEndTime() const
+double STimeRange::PlayEndTime() const
 {
-	return TimeSliderController.IsValid() ? TimeSliderController.Get()->GetPlayRange().GetUpperBoundValue() : 0.f;
+	FFrameNumber UpperBound = MovieScene::DiscreteExclusiveUpper(TimeSliderController->GetPlayRange());
+	return UpperBound.Value;
 }
 
-TOptional<float> STimeRange::MaxViewStartTime() const
+TOptional<double> STimeRange::MaxViewStartTime() const
 {
 	return ViewEndTime();
 }
 
-TOptional<float> STimeRange::MinViewEndTime() const
+TOptional<double> STimeRange::MinViewEndTime() const
 {
 	return ViewStartTime();
 }
 
-TOptional<float> STimeRange::MinPlayStartTime() const
+TOptional<double> STimeRange::MinPlayStartTime() const
 {
 	return WorkingStartTime();
 }
 
-TOptional<float> STimeRange::MaxPlayStartTime() const
+TOptional<double> STimeRange::MaxPlayStartTime() const
 {
 	return PlayEndTime();
 }
 
-TOptional<float> STimeRange::MinPlayEndTime() const
+TOptional<double> STimeRange::MinPlayEndTime() const
 {
 	return PlayStartTime();
 }
 
-TOptional<float> STimeRange::MaxPlayEndTime() const
+TOptional<double> STimeRange::MaxPlayEndTime() const
 {
 	return WorkingEndTime();
 }
 
-TOptional<float> STimeRange::MaxWorkingStartTime() const
+TOptional<double> STimeRange::MaxWorkingStartTime() const
 {
 	return ViewEndTime();
 }
 
-TOptional<float> STimeRange::MinWorkingEndTime() const
+TOptional<double> STimeRange::MinWorkingEndTime() const
 {
 	return ViewStartTime();
 }
 
-void STimeRange::OnWorkingStartTimeCommitted(float NewValue, ETextCommit::Type InTextCommit)
+void STimeRange::OnWorkingStartTimeCommitted(double NewValue, ETextCommit::Type InTextCommit)
 {
 	OnWorkingStartTimeChanged(NewValue);
 }
 
-void STimeRange::OnWorkingEndTimeCommitted(float NewValue, ETextCommit::Type InTextCommit)
+void STimeRange::OnWorkingEndTimeCommitted(double NewValue, ETextCommit::Type InTextCommit)
 {
 	OnWorkingEndTimeChanged(NewValue);
 }
 
-void STimeRange::OnViewStartTimeCommitted(float NewValue, ETextCommit::Type InTextCommit)
+void STimeRange::OnViewStartTimeCommitted(double NewValue, ETextCommit::Type InTextCommit)
 {
 	OnViewStartTimeChanged(NewValue);
 }
 
-void STimeRange::OnViewEndTimeCommitted(float NewValue, ETextCommit::Type InTextCommit)
+void STimeRange::OnViewEndTimeCommitted(double NewValue, ETextCommit::Type InTextCommit)
 {
 	OnViewEndTimeChanged(NewValue);
 }
 
-void STimeRange::OnPlayStartTimeCommitted(float NewValue, ETextCommit::Type InTextCommit)
+void STimeRange::OnPlayStartTimeCommitted(double NewValue, ETextCommit::Type InTextCommit)
 {
 	OnPlayStartTimeChanged(NewValue);
 }
 
-void STimeRange::OnPlayEndTimeCommitted(float NewValue, ETextCommit::Type InTextCommit)
+void STimeRange::OnPlayEndTimeCommitted(double NewValue, ETextCommit::Type InTextCommit)
 {
 	OnPlayEndTimeChanged(NewValue);
 }
 
-void STimeRange::OnWorkingStartTimeChanged(float NewValue)
+void STimeRange::OnWorkingStartTimeChanged(double NewValue)
 {
-	if (auto* Controller = TimeSliderController.Get())
-	{
-		Controller->SetClampRange(NewValue, Controller->GetClampRange().GetUpperBoundValue());
+	FFrameRate PlayRate = TimeSliderController->GetPlayRate();
+	FFrameRate FrameResolution = TimeSliderController->GetFrameResolution();
+	double Time = FrameResolution.AsSeconds(FFrameTime::FromDecimal(NewValue));
 
-		if (NewValue > Controller->GetViewRange().GetLowerBoundValue())
-		{
-			Controller->SetViewRange(NewValue, Controller->GetViewRange().GetUpperBoundValue(), EViewRangeInterpolation::Immediate);
-		}
+	// Clamp range is in seconds
+	TimeSliderController->SetClampRange(Time, TimeSliderController->GetClampRange().GetUpperBoundValue());
+
+	if (Time > TimeSliderController->GetViewRange().GetLowerBoundValue())
+	{
+		TimeSliderController->SetViewRange(Time, TimeSliderController->GetViewRange().GetUpperBoundValue(), EViewRangeInterpolation::Immediate);
 	}
 }
 
-void STimeRange::OnWorkingEndTimeChanged(float NewValue)
+void STimeRange::OnWorkingEndTimeChanged(double NewValue)
 {
-	if (auto* Controller = TimeSliderController.Get())
-	{
-		Controller->SetClampRange(Controller->GetClampRange().GetLowerBoundValue(), NewValue);
+	FFrameRate PlayRate = TimeSliderController->GetPlayRate();
+	FFrameRate FrameResolution = TimeSliderController->GetFrameResolution();
+	double Time = FrameResolution.AsSeconds(FFrameTime::FromDecimal(NewValue));
 
-		if (NewValue < Controller->GetViewRange().GetUpperBoundValue())
-		{
-			Controller->SetViewRange(Controller->GetViewRange().GetLowerBoundValue(), NewValue, EViewRangeInterpolation::Immediate);
-		}
+	// Clamp range is in seconds
+	TimeSliderController->SetClampRange(TimeSliderController->GetClampRange().GetLowerBoundValue(), Time);
+
+	if (Time < TimeSliderController->GetViewRange().GetUpperBoundValue())
+	{
+		TimeSliderController->SetViewRange(TimeSliderController->GetViewRange().GetLowerBoundValue(), Time, EViewRangeInterpolation::Immediate);
 	}
 }
 
-void STimeRange::OnViewStartTimeChanged(float NewValue)
+void STimeRange::OnViewStartTimeChanged(double NewValue)
 {
-	if (auto* Controller = TimeSliderController.Get())
-	{
-		if (NewValue < TimeSliderController.Get()->GetClampRange().GetLowerBoundValue())
-		{
-			Controller->SetClampRange(NewValue, Controller->GetClampRange().GetUpperBoundValue());
-		}
+	FFrameRate FrameResolution = TimeSliderController->GetFrameResolution();
+	double Time = FrameResolution.AsSeconds(FFrameTime::FromDecimal(NewValue));
 
-		Controller->SetViewRange(NewValue, Controller->GetViewRange().GetUpperBoundValue(), EViewRangeInterpolation::Immediate);
+	if (Time < TimeSliderController.Get()->GetClampRange().GetLowerBoundValue())
+	{
+		TimeSliderController->SetClampRange(Time, TimeSliderController->GetClampRange().GetUpperBoundValue());
 	}
+
+	TimeSliderController->SetViewRange(Time, TimeSliderController->GetViewRange().GetUpperBoundValue(), EViewRangeInterpolation::Immediate);
 }
 
-void STimeRange::OnViewEndTimeChanged(float NewValue)
+void STimeRange::OnViewEndTimeChanged(double NewValue)
 {
-	if (auto* Controller = TimeSliderController.Get())
-	{
-		if (NewValue > Controller->GetClampRange().GetUpperBoundValue())
-		{
-			Controller->SetClampRange(Controller->GetClampRange().GetLowerBoundValue(), NewValue);
-		}
+	FFrameRate FrameResolution = TimeSliderController->GetFrameResolution();
+	double Time = FrameResolution.AsSeconds(FFrameTime::FromDecimal(NewValue));
 
-		Controller->SetViewRange(Controller->GetViewRange().GetLowerBoundValue(), NewValue, EViewRangeInterpolation::Immediate);
+	if (Time > TimeSliderController->GetClampRange().GetUpperBoundValue())
+	{
+		TimeSliderController->SetClampRange(TimeSliderController->GetClampRange().GetLowerBoundValue(), Time);
 	}
+
+	TimeSliderController->SetViewRange(TimeSliderController->GetViewRange().GetLowerBoundValue(), Time, EViewRangeInterpolation::Immediate);
 }
 
-void STimeRange::OnPlayStartTimeChanged(float NewValue)
+void STimeRange::OnPlayStartTimeChanged(double NewValue)
 {
-	if (auto* Controller = TimeSliderController.Get())
-	{
-		if (NewValue < TimeSliderController.Get()->GetClampRange().GetLowerBoundValue())
-		{
-			Controller->SetClampRange(NewValue, Controller->GetClampRange().GetUpperBoundValue());
-		}
+	// We can't use the UI control to clamp the value to a Min/Max due to needing an unlimited-range spinbox for
+	// UI adjustment to work in reasonable deltas, so instead we clamp it here.
+	NewValue = FMath::Clamp(NewValue, MinPlayStartTime().GetValue(), MaxPlayStartTime().GetValue());
 
-		Controller->SetPlayRange(NewValue, Controller->GetPlayRange().GetUpperBoundValue());
+	FFrameRate FrameResolution = TimeSliderController->GetFrameResolution();
+	FFrameTime Time = FFrameTime::FromDecimal(NewValue);
+	double     TimeInSeconds = FrameResolution.AsSeconds(Time);
+
+	if (TimeInSeconds < TimeSliderController.Get()->GetClampRange().GetLowerBoundValue())
+	{
+		TimeSliderController->SetClampRange(TimeInSeconds, TimeSliderController->GetClampRange().GetLowerBoundValue());
 	}
+
+	FFrameNumber PlayDuration = MovieScene::DiscreteExclusiveUpper(TimeSliderController->GetPlayRange()) - Time.FrameNumber;
+	TimeSliderController->SetPlayRange(Time.FrameNumber, PlayDuration.Value);
 }
 
-void STimeRange::OnPlayEndTimeChanged(float NewValue)
+void STimeRange::OnPlayEndTimeChanged(double NewValue)
 {
-	if (auto* Controller = TimeSliderController.Get())
-	{
-		if (NewValue > Controller->GetClampRange().GetUpperBoundValue())
-		{
-			Controller->SetClampRange(Controller->GetClampRange().GetLowerBoundValue(), NewValue);
-		}
+	// We can't use the UI control to clamp the value to a Min/Max due to needing an unlimited-range spinbox for
+	// UI adjustment to work in reasonable deltas, so instead we clamp it here.
+	NewValue = FMath::Clamp(NewValue, MinPlayEndTime().GetValue(), MaxPlayEndTime().GetValue());
+
+	FFrameRate FrameResolution = TimeSliderController->GetFrameResolution();
+	FFrameTime Time = FFrameTime::FromDecimal(NewValue);
+	double     TimeInSeconds = FrameResolution.AsSeconds(Time);
 
-		Controller->SetPlayRange(Controller->GetPlayRange().GetLowerBoundValue(), NewValue);
+	if (TimeInSeconds > TimeSliderController->GetClampRange().GetUpperBoundValue())
+	{
+		TimeSliderController->SetClampRange(TimeSliderController->GetClampRange().GetLowerBoundValue(), TimeInSeconds);
 	}
+
+	FFrameNumber StartFrame   = MovieScene::DiscreteInclusiveLower(TimeSliderController->GetPlayRange());
+	int32        PlayDuration = (Time.FrameNumber - StartFrame).Value;
+
+	TimeSliderController->SetPlayRange(StartFrame, PlayDuration);
 }
 
 FText STimeRange::PlayStartTimeTooltip() const
 {
-	bool bShowFrameNumbers = ShowFrameNumbers.IsBound() ? ShowFrameNumbers.Get() : false;
-	if (bShowFrameNumbers)
-	{
-		return LOCTEXT("PlayStartFrameTooltip", "In Frame");
-	}
-	else
-	{
-		return LOCTEXT("PlayStartTimeTooltip", "In Time");
-	}
+	// @todo sequencer-timecode We need a NumericTypeInterface independent way to get the current type of display.
+	// bool bShowFrameNumbers = ShowFrameNumbers.IsBound() ? ShowFrameNumbers.Get() : false;
+	// if (bShowFrameNumbers)
+	// {
+	// 	return LOCTEXT("PlayStartFrameTooltip", "In Frame");
+	// }
+	// else
+	// {
+	// 	return LOCTEXT("PlayStartTimeTooltip", "In Time");
+	// }
+	return LOCTEXT("PlayStartTimeTooltip", "In Time");
 }
 
 FText STimeRange::PlayEndTimeTooltip() const
 {
-	bool bShowFrameNumbers = ShowFrameNumbers.IsBound() ? ShowFrameNumbers.Get() : false;
-	if (bShowFrameNumbers)
-	{
-		return LOCTEXT("PlayEndFrameTooltip", "Out Frame");
-	}
-	else
-	{
-		return LOCTEXT("PlayEndTimeTooltip", "Out Time");
-	}
+	// @todo sequencer-timecode We need a NumericTypeInterface independent way to get the current type of display.
+	// bool bShowFrameNumbers = ShowFrameNumbers.IsBound() ? ShowFrameNumbers.Get() : false;
+	// if (bShowFrameNumbers)
+	// {
+	// 	return LOCTEXT("PlayEndFrameTooltip", "Out Frame");
+	// }
+	// else
+	// {
+	// 	return LOCTEXT("PlayEndTimeTooltip", "Out Time");
+	// }
+	return LOCTEXT("PlayEndTimeTooltip", "Out Time");
 }
 
 FText STimeRange::ViewStartTimeTooltip() const
 {
-	bool bShowFrameNumbers = ShowFrameNumbers.IsBound() ? ShowFrameNumbers.Get() : false;
-	if (bShowFrameNumbers)
-	{
-		return LOCTEXT("ViewStartFrameTooltip", "View Range Start Frame");
-	}
-	else
-	{
-		return LOCTEXT("ViewStartTimeTooltip", "View Range Start Time");
-	}
+	// @todo sequencer-timecode We need a NumericTypeInterface independent way to get the current type of display.
+	// bool bShowFrameNumbers = ShowFrameNumbers.IsBound() ? ShowFrameNumbers.Get() : false;
+	// if (bShowFrameNumbers)
+	// {
+	// 	return LOCTEXT("ViewStartFrameTooltip", "View Range Start Frame");
+	// }
+	// else
+	// {
+	// 	return LOCTEXT("ViewStartTimeTooltip", "View Range Start Time");
+	// }
+	return LOCTEXT("ViewStartTimeTooltip", "View Range Start Time");
+
 }
 
 FText STimeRange::ViewEndTimeTooltip() const
 {
-	bool bShowFrameNumbers = ShowFrameNumbers.IsBound() ? ShowFrameNumbers.Get() : false;
-	if (bShowFrameNumbers)
-	{
-		return LOCTEXT("ViewEndFrameTooltip", "View Range End Frame");
-	}
-	else
-	{
-		return LOCTEXT("ViewEndTimeTooltip", "View Range End Time");
-	}
+	// @todo sequencer-timecode We need a NumericTypeInterface independent way to get the current type of display.
+	// bool bShowFrameNumbers = ShowFrameNumbers.IsBound() ? ShowFrameNumbers.Get() : false;
+	// if (bShowFrameNumbers)
+	// {
+	// 	return LOCTEXT("ViewEndFrameTooltip", "View Range End Frame");
+	// }
+	// else
+	// {
+	// 	return LOCTEXT("ViewEndTimeTooltip", "View Range End Time");
+	// }
+	return LOCTEXT("ViewEndTimeTooltip", "View Range End Time");
 }
 
 #undef LOCTEXT_NAMESPACE
diff --git a/Engine/Source/Editor/SequencerWidgets/Private/STimeRange.h b/Engine/Source/Editor/SequencerWidgets/Private/STimeRange.h
--- a/Engine/Source/Editor/SequencerWidgets/Private/STimeRange.h
+++ b/Engine/Source/Editor/SequencerWidgets/Private/STimeRange.h
@@ -7,6 +7,7 @@
 #include "Widgets/DeclarativeSyntaxSupport.h"
 #include "Widgets/Input/NumericTypeInterface.h"
 #include "ITimeSlider.h"
+#include "QualifiedFrameTime.h"
 
 class STimeRange : public ITimeSlider
 {
@@ -14,8 +15,6 @@ public:
 	SLATE_BEGIN_ARGS(STimeRange)
 		: _ShowWorkingRange(true), _ShowViewRange(false), _ShowPlaybackRange(false)
 	{}
-		/* If we should show frame numbers on the timeline */
-		SLATE_ARGUMENT( TAttribute<bool>, ShowFrameNumbers )
 		/** Whether to show the working range */
 		SLATE_ARGUMENT( bool, ShowWorkingRange )
 		/** Whether to show the view range */
@@ -31,59 +30,61 @@ public:
 	 * 
 	 * @param InArgs   A declaration from which to construct the widget
 	 */
-	void Construct( const FArguments& InArgs, TSharedRef<ITimeSliderController> InTimeSliderController, TSharedRef<INumericTypeInterface<float>> NumericTypeInterface );
+	void Construct( const FArguments& InArgs, TSharedRef<ITimeSliderController> InTimeSliderController, TSharedRef<INumericTypeInterface<double>> NumericTypeInterface );
+
+protected:
+	double GetSpinboxDelta() const;
 	
 protected:
 
-	float PlayStartTime() const;
-	float PlayEndTime() const;
+	double PlayStartTime() const;
+	double PlayEndTime() const;
 
-	TOptional<float> MinPlayStartTime() const;
-	TOptional<float> MaxPlayStartTime() const;
-	TOptional<float> MinPlayEndTime() const;
-	TOptional<float> MaxPlayEndTime() const;
+	TOptional<double> MinPlayStartTime() const;
+	TOptional<double> MaxPlayStartTime() const;
+	TOptional<double> MinPlayEndTime() const;
+	TOptional<double> MaxPlayEndTime() const;
 
-	void OnPlayStartTimeCommitted(float NewValue, ETextCommit::Type InTextCommit);
-	void OnPlayEndTimeCommitted(float NewValue, ETextCommit::Type InTextCommit);
+	void OnPlayStartTimeCommitted(double NewValue, ETextCommit::Type InTextCommit);
+	void OnPlayEndTimeCommitted(double NewValue, ETextCommit::Type InTextCommit);
 
-	void OnPlayStartTimeChanged(float NewValue);
-	void OnPlayEndTimeChanged(float NewValue);
+	void OnPlayStartTimeChanged(double NewValue);
+	void OnPlayEndTimeChanged(double NewValue);
 
 	FText PlayStartTimeTooltip() const;
 	FText PlayEndTimeTooltip() const;
 
 protected:
 
-	float ViewStartTime() const;
-	float ViewEndTime() const;
+	double ViewStartTime() const;
+	double ViewEndTime() const;
 
-	TOptional<float> MaxViewStartTime() const;
-	TOptional<float> MinViewEndTime() const;
+	TOptional<double> MaxViewStartTime() const;
+	TOptional<double> MinViewEndTime() const;
 
-	void OnViewStartTimeCommitted(float NewValue, ETextCommit::Type InTextCommit);
-	void OnViewEndTimeCommitted(float NewValue, ETextCommit::Type InTextCommit);
+	void OnViewStartTimeCommitted(double NewValue, ETextCommit::Type InTextCommit);
+	void OnViewEndTimeCommitted(double NewValue, ETextCommit::Type InTextCommit);
 
-	void OnViewStartTimeChanged(float NewValue);
-	void OnViewEndTimeChanged(float NewValue);
+	void OnViewStartTimeChanged(double NewValue);
+	void OnViewEndTimeChanged(double NewValue);
 
 	FText ViewStartTimeTooltip() const;
 	FText ViewEndTimeTooltip() const;
 
 protected:
 
-	float WorkingStartTime() const;
-	float WorkingEndTime() const;
+	double WorkingStartTime() const;
+	double WorkingEndTime() const;
 
-	TOptional<float> MaxWorkingStartTime() const;
-	TOptional<float> MinWorkingEndTime() const;
+	TOptional<double> MaxWorkingStartTime() const;
+	TOptional<double> MinWorkingEndTime() const;
 
-	void OnWorkingStartTimeCommitted(float NewValue, ETextCommit::Type InTextCommit);
-	void OnWorkingEndTimeCommitted(float NewValue, ETextCommit::Type InTextCommit);
+	void OnWorkingStartTimeCommitted(double NewValue, ETextCommit::Type InTextCommit);
+	void OnWorkingEndTimeCommitted(double NewValue, ETextCommit::Type InTextCommit);
 
-	void OnWorkingStartTimeChanged(float NewValue);
-	void OnWorkingEndTimeChanged(float NewValue);
+	void OnWorkingStartTimeChanged(double NewValue);
+	void OnWorkingEndTimeChanged(double NewValue);
 
 private:
 	TSharedPtr<ITimeSliderController> TimeSliderController;
-	TAttribute<bool> ShowFrameNumbers;
 };
diff --git a/Engine/Source/Editor/SequencerWidgets/Private/STimeRangeSlider.cpp b/Engine/Source/Editor/SequencerWidgets/Private/STimeRangeSlider.cpp
--- a/Engine/Source/Editor/SequencerWidgets/Private/STimeRangeSlider.cpp
+++ b/Engine/Source/Editor/SequencerWidgets/Private/STimeRangeSlider.cpp
@@ -16,8 +16,7 @@ namespace TimeRangeSliderConstants
 void STimeRangeSlider::Construct( const FArguments& InArgs, TSharedRef<ITimeSliderController> InTimeSliderController)
 {
 	TimeSliderController = InTimeSliderController;
-	LastViewRange = TimeSliderController.Get()->GetViewRange();
-	TimeSnapInterval = InArgs._TimeSnapInterval;
+	LastViewRange = TimeSliderController->GetViewRange();
 
 	ResetState();
 	ResetHoveredState();
@@ -30,8 +29,8 @@ float STimeRangeSlider::ComputeDragDelta(const FPointerEvent& MouseEvent, int32
 
 	if (TimeSliderController.IsValid())
 	{
-		StartTime = TimeSliderController.Get()->GetClampRange().GetLowerBoundValue();
-		EndTime = TimeSliderController.Get()->GetClampRange().GetUpperBoundValue();
+		StartTime = TimeSliderController->GetClampRange().GetLowerBoundValue();
+		EndTime = TimeSliderController->GetClampRange().GetUpperBoundValue();
 	}
 	float DragDistance = (MouseEvent.GetScreenSpacePosition() - MouseDownPosition).X;
 
@@ -48,10 +47,10 @@ void STimeRangeSlider::ComputeHandleOffsets(float& LeftHandleOffset, float& Hand
 
 	if (TimeSliderController.IsValid())
 	{
-		StartTime = TimeSliderController.Get()->GetClampRange().GetLowerBoundValue();
-		InTime = TimeSliderController.Get()->GetViewRange().GetLowerBoundValue();
-		OutTime = TimeSliderController.Get()->GetViewRange().GetUpperBoundValue();
-		EndTime = TimeSliderController.Get()->GetClampRange().GetUpperBoundValue();
+		StartTime = TimeSliderController->GetClampRange().GetLowerBoundValue();
+		InTime = TimeSliderController->GetViewRange().GetLowerBoundValue();
+		OutTime = TimeSliderController->GetViewRange().GetUpperBoundValue();
+		EndTime = TimeSliderController->GetClampRange().GetUpperBoundValue();
 	}
 
 	const float UnitsToPixel = float(GeometryWidth - TimeRangeSliderConstants::HandleSize*2) / (EndTime - StartTime);
@@ -131,7 +130,7 @@ FReply STimeRangeSlider::OnMouseButtonDown( const FGeometry& MyGeometry, const F
 	MouseDownPosition = MouseEvent.GetScreenSpacePosition();
 	if (TimeSliderController.IsValid())
 	{		
-		MouseDownViewRange = TimeSliderController.Get()->GetViewRange();
+		MouseDownViewRange = TimeSliderController->GetViewRange();
 	}
 
 	if (bHandleHovered)
@@ -166,8 +165,6 @@ FReply STimeRangeSlider::OnMouseMove( const FGeometry& MyGeometry, const FPointe
 	{
 		float DragDelta = ComputeDragDelta(MouseEvent, MyGeometry.GetLocalSize().X);
 
-		const float SnapInterval = TimeSnapInterval.Get(1.f);
-
 		ITimeSliderController* TimeSliderControllerPtr = TimeSliderController.Get();
 		if (!TimeSliderControllerPtr)
 		{
@@ -176,10 +173,10 @@ FReply STimeRangeSlider::OnMouseMove( const FGeometry& MyGeometry, const FPointe
 
 		if (bHandleDragged)
 		{
-			float NewIn = MouseDownViewRange.GetLowerBoundValue() + DragDelta;
-			float NewOut = MouseDownViewRange.GetUpperBoundValue() + DragDelta;
+			double NewIn = MouseDownViewRange.GetLowerBoundValue() + DragDelta;
+			double NewOut = MouseDownViewRange.GetUpperBoundValue() + DragDelta;
 
-			TRange<float> ClampRange = TimeSliderControllerPtr->GetClampRange();
+			TRange<double> ClampRange = TimeSliderControllerPtr->GetClampRange();
 			if (NewIn < ClampRange.GetLowerBoundValue())
 			{
 				NewIn = ClampRange.GetLowerBoundValue();
@@ -193,33 +190,15 @@ FReply STimeRangeSlider::OnMouseMove( const FGeometry& MyGeometry, const FPointe
 			
 			TimeSliderControllerPtr->SetViewRange(NewIn, NewOut, EViewRangeInterpolation::Immediate);
 		}
-		else if (bLeftHandleDragged)
-		{
-			float NewIn = MouseDownViewRange.GetLowerBoundValue() + DragDelta;
-			float MaxIn = TimeSliderControllerPtr->GetViewRange().GetUpperBoundValue() - SnapInterval;
-
-			// Clamp to the upper view range bound minus some small number to prevent 0-sized (or negative) ranges
-			NewIn = FMath::Min(NewIn, MaxIn);
-
-			// Double check we're not outside the clamp range
-			TRange<float> ClampRange = TimeSliderControllerPtr->GetClampRange();
-			NewIn = FMath::Clamp(NewIn, ClampRange.GetLowerBoundValue(), ClampRange.GetUpperBoundValue());
-
-			TimeSliderControllerPtr->SetViewRange(NewIn, TimeSliderControllerPtr->GetViewRange().GetUpperBoundValue(), EViewRangeInterpolation::Immediate);
-		}
-		else if (bRightHandleDragged)
+		else if (bLeftHandleDragged || bRightHandleDragged)
 		{
-			float NewOut = MouseDownViewRange.GetUpperBoundValue() + DragDelta;
-			float MinOut = TimeSliderControllerPtr->GetViewRange().GetLowerBoundValue() + SnapInterval;
+			double InDelta  = bLeftHandleDragged  ? DragDelta : 0.0;
+			double OutDelta = bRightHandleDragged ? DragDelta : 0.0;
 
-			// Clamp to the upper view range bound minus some small number to prevent 0-sized (or negative) ranges
-			NewOut = FMath::Max(NewOut, MinOut);
+			double NewIn  = MouseDownViewRange.GetLowerBoundValue() + InDelta;
+			double NewOut = MouseDownViewRange.GetUpperBoundValue() + OutDelta;
 
-			// Double check we're not outside the clamp range
-			TRange<float> ClampRange = TimeSliderControllerPtr->GetClampRange();
-			NewOut = FMath::Clamp(NewOut, ClampRange.GetLowerBoundValue(), ClampRange.GetUpperBoundValue());
-
-			TimeSliderControllerPtr->SetViewRange(TimeSliderController.Get()->GetViewRange().GetLowerBoundValue(), NewOut, EViewRangeInterpolation::Immediate);
+			TimeSliderControllerPtr->SetViewRange(NewIn, NewOut, EViewRangeInterpolation::Immediate);
 		}
 
 		return FReply::Handled();
@@ -233,9 +212,9 @@ FReply STimeRangeSlider::OnMouseMove( const FGeometry& MyGeometry, const FPointe
 		float RightHandleOffset = 0.f;
 		ComputeHandleOffsets(LeftHandleOffset, HandleOffset, RightHandleOffset, MyGeometry.GetLocalSize().X);
 		
-		FGeometry LeftHandleRect = MyGeometry.MakeChild(FVector2D(LeftHandleOffset, 0.f), FVector2D(TimeRangeSliderConstants::HandleSize, TimeRangeSliderConstants::HandleSize));
+		FGeometry LeftHandleRect  = MyGeometry.MakeChild(FVector2D(LeftHandleOffset, 0.f), FVector2D(TimeRangeSliderConstants::HandleSize, TimeRangeSliderConstants::HandleSize));
 		FGeometry RightHandleRect = MyGeometry.MakeChild(FVector2D(RightHandleOffset, 0.f), FVector2D(TimeRangeSliderConstants::HandleSize, TimeRangeSliderConstants::HandleSize));
-		FGeometry HandleRect = MyGeometry.MakeChild(FVector2D(HandleOffset, 0.0f), FVector2D(RightHandleOffset-LeftHandleOffset-TimeRangeSliderConstants::HandleSize, TimeRangeSliderConstants::HandleSize));
+		FGeometry HandleRect      = MyGeometry.MakeChild(FVector2D(HandleOffset, 0.0f), FVector2D(RightHandleOffset-LeftHandleOffset-TimeRangeSliderConstants::HandleSize, TimeRangeSliderConstants::HandleSize));
 
 		FVector2D LocalMousePosition = MouseEvent.GetScreenSpacePosition();
 
@@ -275,18 +254,18 @@ FReply STimeRangeSlider::OnMouseButtonDoubleClick( const FGeometry& MyGeometry,
 	{
 		if (TimeSliderController.IsValid())
 		{
-			if (FMath::IsNearlyEqual(TimeSliderController.Get()->GetViewRange().GetLowerBoundValue(), TimeSliderController.Get()->GetClampRange().GetLowerBoundValue()) &&
-				FMath::IsNearlyEqual(TimeSliderController.Get()->GetViewRange().GetUpperBoundValue(), TimeSliderController.Get()->GetClampRange().GetUpperBoundValue()))
+			if (FMath::IsNearlyEqual(TimeSliderController->GetViewRange().GetLowerBoundValue(), TimeSliderController->GetClampRange().GetLowerBoundValue()) &&
+				FMath::IsNearlyEqual(TimeSliderController->GetViewRange().GetUpperBoundValue(), TimeSliderController->GetClampRange().GetUpperBoundValue()))
 			{
 				if (!LastViewRange.IsEmpty())
 				{
-					TimeSliderController.Get()->SetViewRange(LastViewRange.GetLowerBoundValue(), LastViewRange.GetUpperBoundValue(), EViewRangeInterpolation::Immediate);
+					TimeSliderController->SetViewRange(LastViewRange.GetLowerBoundValue(), LastViewRange.GetUpperBoundValue(), EViewRangeInterpolation::Immediate);
 				}
 			}
 			else
 			{
-				LastViewRange = TRange<float>(TimeSliderController.Get()->GetViewRange().GetLowerBoundValue(), TimeSliderController.Get()->GetViewRange().GetUpperBoundValue());
-				TimeSliderController.Get()->SetViewRange(TimeSliderController.Get()->GetClampRange().GetLowerBoundValue(), TimeSliderController.Get()->GetClampRange().GetUpperBoundValue(), EViewRangeInterpolation::Immediate);
+				LastViewRange = TimeSliderController->GetViewRange();
+				TimeSliderController->SetViewRange(TimeSliderController->GetClampRange().GetLowerBoundValue(), TimeSliderController->GetClampRange().GetUpperBoundValue(), EViewRangeInterpolation::Immediate);
 			}
 		}
 
diff --git a/Engine/Source/Editor/SequencerWidgets/Private/STimeRangeSlider.h b/Engine/Source/Editor/SequencerWidgets/Private/STimeRangeSlider.h
--- a/Engine/Source/Editor/SequencerWidgets/Private/STimeRangeSlider.h
+++ b/Engine/Source/Editor/SequencerWidgets/Private/STimeRangeSlider.h
@@ -16,8 +16,6 @@ class STimeRangeSlider : public SCompoundWidget
 {
 public:
 	SLATE_BEGIN_ARGS(STimeRangeSlider){}
-		/* The time snap interval for the timeline */
-		SLATE_ARGUMENT( TAttribute<float>, TimeSnapInterval )
 		SLATE_DEFAULT_SLOT(FArguments, Content)
 	SLATE_END_ARGS()
 
@@ -57,12 +55,10 @@ private:
 	FVector2D MouseDownPosition;
 
 	/* The in/out view range on mouse down */
-	TRange<float> MouseDownViewRange;
+	TRange<double> MouseDownViewRange;
 
 	/* The in/out view range viewed before expansion */
-	TRange<float> LastViewRange;
-
-	TAttribute<float> TimeSnapInterval;
+	TRange<double> LastViewRange;
 
 	TSharedPtr<ITimeSliderController> TimeSliderController;
 };
diff --git a/Engine/Source/Editor/SequencerWidgets/Private/SequencerWidgetsModule.cpp b/Engine/Source/Editor/SequencerWidgets/Private/SequencerWidgetsModule.cpp
--- a/Engine/Source/Editor/SequencerWidgets/Private/SequencerWidgetsModule.cpp
+++ b/Engine/Source/Editor/SequencerWidgets/Private/SequencerWidgetsModule.cpp
@@ -35,17 +35,15 @@ public:
 			.MirrorLabels(bMirrorLabels);
 	}
 
-	TSharedRef<SWidget> CreateTimeRangeSlider( const TSharedRef<class ITimeSliderController>& InController, const TAttribute<float>& InTimeSnapIntervalDelegate ) override
+	TSharedRef<SWidget> CreateTimeRangeSlider( const TSharedRef<class ITimeSliderController>& InController ) override
 	{
-		return SNew( STimeRangeSlider, InController )
-			.TimeSnapInterval( InTimeSnapIntervalDelegate );
+		return SNew( STimeRangeSlider, InController );
 	}
 
 	TSharedRef<ITimeSlider> CreateTimeRange(const FTimeRangeArgs& InArgs, const TSharedRef<SWidget>& Content) override
 	{
 		return SNew( STimeRange, InArgs.Controller, InArgs.NumericTypeInterface)
 		.Visibility(InArgs.VisibilityDelegate)
-		.ShowFrameNumbers(InArgs.ShowFrameNumbersDelegate)
 		.ShowWorkingRange(!!(InArgs.ShowRanges & EShowRange::WorkingRange))
 		.ShowViewRange(!!(InArgs.ShowRanges & EShowRange::ViewRange))
 		.ShowPlaybackRange(!!(InArgs.ShowRanges & EShowRange::PlaybackRange))
diff --git a/Engine/Source/Editor/SequencerWidgets/Public/ISequencerWidgetsModule.h b/Engine/Source/Editor/SequencerWidgets/Public/ISequencerWidgetsModule.h
--- a/Engine/Source/Editor/SequencerWidgets/Public/ISequencerWidgetsModule.h
+++ b/Engine/Source/Editor/SequencerWidgets/Public/ISequencerWidgetsModule.h
@@ -10,6 +10,7 @@
 #include "ITimeSlider.h"
 #include "Widgets/Input/NumericTypeInterface.h"
 #include "Modules/ModuleInterface.h"
+#include "QualifiedFrameTime.h"
 
 /** Enum denoting which time ranges to display on a time range */
 enum class EShowRange
@@ -28,13 +29,11 @@ struct FTimeRangeArgs
 		EShowRange InShowRanges,
 		TSharedRef<ITimeSliderController> InController,
 		TAttribute<EVisibility> InVisibilityDelegate,
-		TAttribute<bool> InShowFrameNumbersDelegate,
-		TSharedRef<INumericTypeInterface<float>> InNumericTypeInterface
+		TSharedRef<INumericTypeInterface<double>> InNumericTypeInterface
 		)
 		: ShowRanges(InShowRanges)
 		, Controller(InController)
 		, VisibilityDelegate(InVisibilityDelegate)
-		, ShowFrameNumbersDelegate(InShowFrameNumbersDelegate)
 		, NumericTypeInterface(InNumericTypeInterface)
 	{
 	}
@@ -45,10 +44,8 @@ struct FTimeRangeArgs
 	TSharedRef<ITimeSliderController> Controller;
 	/** Visibility delegate */
 	TAttribute<EVisibility> VisibilityDelegate;
-	/** Whether to show frame numbers */
-	TAttribute<bool> ShowFrameNumbersDelegate;
 	/** Numeric type interface to use for frame<->time conversion */
-	TSharedRef<INumericTypeInterface<float>> NumericTypeInterface;
+	TSharedRef<INumericTypeInterface<double>> NumericTypeInterface;
 };
 
 /**
@@ -61,6 +58,6 @@ public:
 
 	virtual TSharedRef<ITimeSlider> CreateTimeSlider( const TSharedRef<ITimeSliderController>& InController, bool bMirrorLabels  ) = 0;
 	virtual TSharedRef<ITimeSlider> CreateTimeSlider( const TSharedRef<ITimeSliderController>& InController, const TAttribute<EVisibility>& VisibilityDelegate, bool bMirrorLabels  ) = 0;
-	virtual TSharedRef<SWidget> CreateTimeRangeSlider( const TSharedRef<class ITimeSliderController>& InController, const TAttribute<float>& InTimeSnapIntervalDelegate ) = 0;
+	virtual TSharedRef<SWidget> CreateTimeRangeSlider( const TSharedRef<class ITimeSliderController>& InController ) = 0;
 	virtual TSharedRef<ITimeSlider> CreateTimeRange( const FTimeRangeArgs& InArgs, const TSharedRef<SWidget>& Content ) = 0;
 };
diff --git a/Engine/Source/Editor/SequencerWidgets/Public/ITimeSlider.h b/Engine/Source/Editor/SequencerWidgets/Public/ITimeSlider.h
--- a/Engine/Source/Editor/SequencerWidgets/Public/ITimeSlider.h
+++ b/Engine/Source/Editor/SequencerWidgets/Public/ITimeSlider.h
@@ -8,6 +8,7 @@
 #include "Widgets/SWidget.h"
 #include "Widgets/DeclarativeSyntaxSupport.h"
 #include "Widgets/SCompoundWidget.h"
+#include "Widgets/Input/NumericTypeInterface.h"
 #include "Editor/Sequencer/Public/ISequencerInputHandler.h"
 #include "IMovieScenePlayer.h"
 
@@ -23,43 +24,44 @@ enum class EViewRangeInterpolation
 	Immediate,
 };
 
-DECLARE_DELEGATE_TwoParams( FOnScrubPositionChanged, float, bool )
-DECLARE_DELEGATE_TwoParams( FOnViewRangeChanged, TRange<float>, EViewRangeInterpolation )
-DECLARE_DELEGATE_OneParam( FOnRangeChanged, TRange<float> )
-DECLARE_DELEGATE_RetVal_OneParam( float, FOnGetNearestKey, float )
+DECLARE_DELEGATE_TwoParams( FOnScrubPositionChanged, FFrameTime, bool )
+DECLARE_DELEGATE_TwoParams( FOnViewRangeChanged, TRange<double>, EViewRangeInterpolation )
+DECLARE_DELEGATE_OneParam( FOnTimeRangeChanged, TRange<double> )
+DECLARE_DELEGATE_OneParam( FOnFrameRangeChanged, TRange<FFrameNumber> )
+DECLARE_DELEGATE_RetVal_OneParam( FFrameNumber, FOnGetNearestKey, FFrameTime )
 
 /** Structure used to wrap up a range, and an optional animation target */
-struct FAnimatedRange : public TRange<float>
+struct FAnimatedRange : public TRange<double>
 {
 	/** Default Construction */
 	FAnimatedRange() : TRange() {}
 	/** Construction from a lower and upper bound */
-	FAnimatedRange( float LowerBound, float UpperBound ) : TRange( LowerBound, UpperBound ) {}
+	FAnimatedRange( double LowerBound, double UpperBound ) : TRange( LowerBound, UpperBound ) {}
 	/** Copy-construction from simple range */
-	FAnimatedRange( const TRange<float>& InRange ) : TRange(InRange) {}
+	FAnimatedRange( const TRange<double>& InRange ) : TRange(InRange) {}
 
 	/** Helper function to wrap an attribute to an animated range with a non-animated one */
-	static TAttribute<TRange<float>> WrapAttribute( const TAttribute<FAnimatedRange>& InAttribute )
+	static TAttribute<TRange<double>> WrapAttribute( const TAttribute<FAnimatedRange>& InAttribute )
 	{
-		typedef TAttribute<TRange<float>> Attr;
+		typedef TAttribute<TRange<double>> Attr;
 		return Attr::Create(Attr::FGetter::CreateLambda([=](){ return InAttribute.Get(); }));
 	}
 
 	/** Helper function to wrap an attribute to a non-animated range with an animated one */
-	static TAttribute<FAnimatedRange> WrapAttribute( const TAttribute<TRange<float>>& InAttribute )
+	static TAttribute<FAnimatedRange> WrapAttribute( const TAttribute<TRange<double>>& InAttribute )
 	{
 		typedef TAttribute<FAnimatedRange> Attr;
 		return Attr::Create(Attr::FGetter::CreateLambda([=](){ return InAttribute.Get(); }));
 	}
 
 	/** Get the current animation target, or the whole view range when not animating */
-	const TRange<float>& GetAnimationTarget() const
+	const TRange<double>& GetAnimationTarget() const
 	{
 		return AnimationTarget.IsSet() ? AnimationTarget.GetValue() : *this;
 	}
 	
 	/** The animation target, if animating */
-	TOptional<TRange<float>> AnimationTarget;
+	TOptional<TRange<double>> AnimationTarget;
 };
 
 struct FTimeSliderArgs
@@ -68,18 +70,12 @@ struct FTimeSliderArgs
 		: ScrubPosition(0)
 		, ViewRange( FAnimatedRange(0.0f, 5.0f) )
 		, ClampRange( FAnimatedRange(-FLT_MAX/2.f, FLT_MAX/2.f) )
-		, OnScrubPositionChanged()
-		, OnBeginScrubberMovement()
-		, OnEndScrubberMovement()
-		, OnViewRangeChanged()
-		, OnClampRangeChanged()
-		, OnGetNearestKey()
 		, AllowZoom(true)
 		, Settings(nullptr)
 	{}
 
 	/** The scrub position */
-	TAttribute<float> ScrubPosition;
+	TAttribute<FFrameTime> ScrubPosition;
 
 	/** View time range */
 	TAttribute< FAnimatedRange > ViewRange;
@@ -100,19 +96,25 @@ struct FTimeSliderArgs
 	FOnViewRangeChanged OnViewRangeChanged;
 
 	/** Called when the clamp range changes */
-	FOnRangeChanged OnClampRangeChanged;
+	FOnTimeRangeChanged OnClampRangeChanged;
 
 	/** Delegate that is called when getting the nearest key */
 	FOnGetNearestKey OnGetNearestKey;
 
 	/** Attribute defining the active sub-sequence range for this controller */
-	TAttribute<TOptional<TRange<float>>> SubSequenceRange;
+	TAttribute<TOptional<TRange<FFrameNumber>>> SubSequenceRange;
 
 	/** Attribute defining the playback range for this controller */
-	TAttribute<TRange<float>> PlaybackRange;
+	TAttribute<TRange<FFrameNumber>> PlaybackRange;
+
+	/** Attribute for the current sequence's playback rate */
+	TAttribute<FFrameRate> PlayRate;
+
+	/** Attribute for the current sequence's frame resolution */
+	TAttribute<FFrameRate> FrameResolution;
 
 	/** Delegate that is called when the playback range wants to change */
-	FOnRangeChanged OnPlaybackRangeChanged;
+	FOnFrameRangeChanged OnPlaybackRangeChanged;
 
 	/** Called right before the playback range starts to be dragged */
 	FSimpleDelegate OnPlaybackRangeBeginDrag;
@@ -121,10 +123,10 @@ struct FTimeSliderArgs
 	FSimpleDelegate OnPlaybackRangeEndDrag;
 
 	/** Attribute defining the selection range for this controller */
-	TAttribute<TRange<float>> SelectionRange;
+	TAttribute<TRange<FFrameNumber>> SelectionRange;
 
 	/** Delegate that is called when the selection range wants to change */
-	FOnRangeChanged OnSelectionRangeChanged;
+	FOnFrameRangeChanged OnSelectionRangeChanged;
 
 	/** Called right before the selection range starts to be dragged */
 	FSimpleDelegate OnSelectionRangeBeginDrag;
@@ -148,6 +150,9 @@ struct FTimeSliderArgs
 
 	/** User-supplied settings object */
 	USequencerSettings* Settings;
+
+	/** Numeric Type interface for converting between frame numbers and display formats. */
+	TSharedPtr<INumericTypeInterface<double>> NumericTypeInterface;
 };
 
 class ITimeSliderController : public ISequencerInputHandler
@@ -157,6 +162,12 @@ public:
 	virtual int32 OnPaintTimeSlider( bool bMirrorLabels, const FGeometry& AllottedGeometry, const FSlateRect& MyCullingRect, FSlateWindowElementList& OutDrawElements, int32 LayerId, const FWidgetStyle& InWidgetStyle, bool bParentEnabled ) const = 0;
 	virtual FCursorReply OnCursorQuery( TSharedRef<const SWidget> WidgetOwner, const FGeometry& MyGeometry, const FPointerEvent& CursorEvent ) const = 0;
 
+	/** Get the current play rate for this controller */
+	virtual FFrameRate GetPlayRate() const = 0;
+
+	/** Get the current frame resolution for this controller */
+	virtual FFrameRate GetFrameResolution() const = 0;
+
 	/** Get the current view range for this controller */
 	virtual FAnimatedRange GetViewRange() const { return FAnimatedRange(); }
 
@@ -164,13 +175,7 @@ public:
 	virtual FAnimatedRange GetClampRange() const { return FAnimatedRange(); }
 
 	/** Get the current play range for this controller */
-	virtual TRange<float> GetPlayRange() const { return TRange<float>(); }
-
-	/** Convert a time to a frame */
-	virtual int32 TimeToFrame(float Time) const { return 1; }
-
-	/** Convert a time to a frame */
-	virtual float FrameToTime(int32 Frame) const { return 1.f; }
+	virtual TRange<FFrameNumber> GetPlayRange() const { return TRange<FFrameNumber>(); }
 
 	/**
 	 * Set a new range based on a min, max and an interpolation mode
@@ -179,7 +184,7 @@ public:
 	 * @param NewRangeMax		The new upper bound of the range
 	 * @param Interpolation		How to set the new range (either immediately, or animated)
 	 */
-	virtual void SetViewRange( float NewRangeMin, float NewRangeMax, EViewRangeInterpolation Interpolation ) {}
+	virtual void SetViewRange( double NewRangeMin, double NewRangeMax, EViewRangeInterpolation Interpolation ) {}
 
 	/**
 	 * Set a new clamp range based on a min, max
@@ -187,15 +192,15 @@ public:
 	 * @param NewRangeMin		The new lower bound of the clamp range
 	 * @param NewRangeMax		The new upper bound of the clamp range
 	 */
-	virtual void SetClampRange( float NewRangeMin, float NewRangeMax) {}
+	virtual void SetClampRange( double NewRangeMin, double NewRangeMax) {}
 
 	/**
 	 * Set a new playback range based on a min, max
 	 * 
-	 * @param NewRangeMin		The new lower bound of the playback range
-	 * @param NewRangeMax		The new upper bound of the playback range
+	 * @param RangeStart		The new lower bound of the playback range
+	 * @param RangeDuration		The total number of frames that we play for
 	 */
-	virtual void SetPlayRange( float NewRangeMin, float NewRangeMax ) {}
+	virtual void SetPlayRange( FFrameNumber RangeStart, int32 RangeDuration ) {}
 };
 
 /**
diff --git a/Engine/Source/Editor/UMGEditor/Private/Animation/MarginTrackEditor.cpp b/Engine/Source/Editor/UMGEditor/Private/Animation/MarginTrackEditor.cpp
--- a/Engine/Source/Editor/UMGEditor/Private/Animation/MarginTrackEditor.cpp
+++ b/Engine/Source/Editor/UMGEditor/Private/Animation/MarginTrackEditor.cpp
@@ -1,9 +1,7 @@
 // Copyright 1998-2018 Epic Games, Inc. All Rights Reserved.
 
 #include "Animation/MarginTrackEditor.h"
-#include "PropertySection.h"
 #include "ISectionLayoutBuilder.h"
-#include "FloatCurveKeyArea.h"
 
 
 FName FMarginTrackEditor::LeftName( "Left" );
@@ -12,99 +10,26 @@ FName FMarginTrackEditor::RightName( "Right" );
 FName FMarginTrackEditor::BottomName( "Bottom" );
 
 
-class FMarginPropertySection
-	: public FPropertySection
-{
-public:
-
-	FMarginPropertySection(ISequencer* InSequencer, FGuid InObjectBinding, FName InPropertyName, const FString& InPropertyPath, UMovieSceneSection& InSectionObject, const FText& SectionName)
-		: FPropertySection(InSequencer, InObjectBinding, InPropertyName, InPropertyPath, InSectionObject, SectionName)
-	{
-	}
-
-	virtual void GenerateSectionLayout( class ISectionLayoutBuilder& LayoutBuilder ) const override
-	{
-		UMovieSceneMarginSection* MarginSection = Cast<UMovieSceneMarginSection>(&SectionObject);
-
-		TAttribute<TOptional<float>> MarginLeftExternalValue = TAttribute<TOptional<float>>::Create(
-			TAttribute<TOptional<float>>::FGetter::CreateRaw(this, &FMarginPropertySection::GetMarginLeftValue));
-		TSharedRef<FFloatCurveKeyArea> LeftKeyArea = MakeShareable(new FFloatCurveKeyArea(&MarginSection->GetLeftCurve(), MarginLeftExternalValue, MarginSection));
-
-		TAttribute<TOptional<float>> MarginTopExternalValue = TAttribute<TOptional<float>>::Create(
-			TAttribute<TOptional<float>>::FGetter::CreateRaw(this, &FMarginPropertySection::GetMarginTopValue));
-		TSharedRef<FFloatCurveKeyArea> TopKeyArea = MakeShareable(new FFloatCurveKeyArea(&MarginSection->GetTopCurve(), MarginTopExternalValue, MarginSection));
-
-		TAttribute<TOptional<float>> MarginRightExternalValue = TAttribute<TOptional<float>>::Create(
-			TAttribute<TOptional<float>>::FGetter::CreateRaw(this, &FMarginPropertySection::GetMarginRightValue));
-		TSharedRef<FFloatCurveKeyArea> RightKeyArea = MakeShareable(new FFloatCurveKeyArea(&MarginSection->GetRightCurve(), MarginRightExternalValue, MarginSection));
-
-		TAttribute<TOptional<float>> MarginBottomExternalValue = TAttribute<TOptional<float>>::Create(
-			TAttribute<TOptional<float>>::FGetter::CreateRaw(this, &FMarginPropertySection::GetMarginBottomValue));
-		TSharedRef<FFloatCurveKeyArea> BottomKeyArea = MakeShareable(new FFloatCurveKeyArea( &MarginSection->GetBottomCurve(), MarginBottomExternalValue, MarginSection));
-
-		LayoutBuilder.AddKeyArea("Left", NSLOCTEXT("FMarginPropertySection", "MarginLeft", "Left"), LeftKeyArea);
-		LayoutBuilder.AddKeyArea("Top", NSLOCTEXT("FMarginPropertySection", "MarginTop", "Top"), TopKeyArea);
-		LayoutBuilder.AddKeyArea("Right", NSLOCTEXT("FMarginPropertySection", "MarginRight", "Right"), RightKeyArea);
-		LayoutBuilder.AddKeyArea("Bottom", NSLOCTEXT("FMarginPropertySection", "MarginBottom", "Bottom"), BottomKeyArea);
-	}
-
-
-private:
-
-	TOptional<float> GetMarginLeftValue() const
-	{
-		TOptional<FMargin> Margin = GetPropertyValue<FMargin>();
-		return Margin.IsSet() ? TOptional<float>(Margin.GetValue().Left) : TOptional<float>();
-	}
-
-	TOptional<float> GetMarginTopValue() const
-	{
-		TOptional<FMargin> Margin = GetPropertyValue<FMargin>();
-		return Margin.IsSet() ? TOptional<float>(Margin.GetValue().Top) : TOptional<float>();
-	}
-
-	TOptional<float> GetMarginRightValue() const
-	{
-		TOptional<FMargin> Margin = GetPropertyValue<FMargin>();
-		return Margin.IsSet() ? TOptional<float>(Margin.GetValue().Right) : TOptional<float>();
-	}
-
-	TOptional<float> GetMarginBottomValue() const
-	{
-		TOptional<FMargin> Margin = GetPropertyValue<FMargin>();
-		return Margin.IsSet() ? TOptional<float>(Margin.GetValue().Bottom) : TOptional<float>();
-	}
-};
-
-
 TSharedRef<ISequencerTrackEditor> FMarginTrackEditor::CreateTrackEditor( TSharedRef<ISequencer> InSequencer )
 {
 	return MakeShareable( new FMarginTrackEditor( InSequencer ) );
 }
 
 
-TSharedRef<ISequencerSection> FMarginTrackEditor::MakeSectionInterface( UMovieSceneSection& SectionObject, UMovieSceneTrack& Track, FGuid ObjectBinding )
+void FMarginTrackEditor::GenerateKeysFromPropertyChanged( const FPropertyChangedParams& PropertyChangedParams, FGeneratedTrackKeys& OutGeneratedKeys)
 {
-	UMovieScenePropertyTrack* PropertyTrack = Cast<UMovieScenePropertyTrack>(&Track);
-	checkf(PropertyTrack != nullptr, TEXT("Incompatible track in FMarginTrackEditor"));
-	return MakeShareable(new FMarginPropertySection(GetSequencer().Get(), ObjectBinding, PropertyTrack->GetPropertyName(), PropertyTrack->GetPropertyPath(), SectionObject, Track.GetDisplayName()));
-}
-
+	FPropertyPath StructPath = PropertyChangedParams.StructPathToKey;
+	FName ChannelName = StructPath.GetNumProperties() != 0 ? StructPath.GetLeafMostProperty().Property->GetFName() : NAME_None;
 
-void FMarginTrackEditor::GenerateKeysFromPropertyChanged( const FPropertyChangedParams& PropertyChangedParams, TArray<FMarginKey>& NewGeneratedKeys, TArray<FMarginKey>& DefaultGeneratedKeys)
-{
-	FName ChannelName = PropertyChangedParams.StructPropertyNameToKey;
 	FMargin Margin = PropertyChangedParams.GetPropertyValue<FMargin>();
 
-	TArray<FMarginKey>& LeftKeys = ChannelName == NAME_None || ChannelName == LeftName ? NewGeneratedKeys : DefaultGeneratedKeys;
-	LeftKeys.Add( FMarginKey( EKeyMarginChannel::Left, Margin.Left ) );
-
-	TArray<FMarginKey>& TopKeys = ChannelName == NAME_None || ChannelName == TopName ? NewGeneratedKeys : DefaultGeneratedKeys;
-	TopKeys.Add( FMarginKey( EKeyMarginChannel::Top, Margin.Top ) );
-
-	TArray<FMarginKey>& RightKeys = ChannelName == NAME_None || ChannelName == RightName ? NewGeneratedKeys : DefaultGeneratedKeys;
-	RightKeys.Add( FMarginKey( EKeyMarginChannel::Right, Margin.Right ) );
+	const bool bKeyLeft   = ChannelName == NAME_None || ChannelName == LeftName;
+	const bool bKeyTop    = ChannelName == NAME_None || ChannelName == TopName;
+	const bool bKeyRight  = ChannelName == NAME_None || ChannelName == RightName;
+	const bool bKeyBottom = ChannelName == NAME_None || ChannelName == BottomName;
 
-	TArray<FMarginKey>& BottomKeys = ChannelName == NAME_None || ChannelName == BottomName ? NewGeneratedKeys : DefaultGeneratedKeys;
-	BottomKeys.Add( FMarginKey( EKeyMarginChannel::Bottom, Margin.Bottom ) );
+	OutGeneratedKeys.Add(FMovieSceneChannelValueSetter::Create<FMovieSceneFloatChannel>(0, Margin.Left,   bKeyLeft));
+	OutGeneratedKeys.Add(FMovieSceneChannelValueSetter::Create<FMovieSceneFloatChannel>(1, Margin.Top,    bKeyTop));
+	OutGeneratedKeys.Add(FMovieSceneChannelValueSetter::Create<FMovieSceneFloatChannel>(2, Margin.Right,  bKeyRight));
+	OutGeneratedKeys.Add(FMovieSceneChannelValueSetter::Create<FMovieSceneFloatChannel>(3, Margin.Bottom, bKeyBottom));
 }
diff --git a/Engine/Source/Editor/UMGEditor/Private/Animation/MarginTrackEditor.h b/Engine/Source/Editor/UMGEditor/Private/Animation/MarginTrackEditor.h
--- a/Engine/Source/Editor/UMGEditor/Private/Animation/MarginTrackEditor.h
+++ b/Engine/Source/Editor/UMGEditor/Private/Animation/MarginTrackEditor.h
@@ -12,7 +12,7 @@
 #include "Animation/MovieSceneMarginSection.h"
 
 class FMarginTrackEditor
-	: public FPropertyTrackEditor<UMovieSceneMarginTrack, UMovieSceneMarginSection, FMarginKey>
+	: public FPropertyTrackEditor<UMovieSceneMarginTrack>
 {
 public:
 
@@ -22,7 +22,7 @@ public:
 	 * @param InSequencer The sequencer instance to be used by this tool
 	 */
 	FMarginTrackEditor( TSharedRef<ISequencer> InSequencer )
-		: FPropertyTrackEditor<UMovieSceneMarginTrack, UMovieSceneMarginSection, FMarginKey>( InSequencer, GetAnimatedPropertyTypes() )
+		: FPropertyTrackEditor( InSequencer, GetAnimatedPropertyTypes() )
 	{
 	}
 
@@ -44,15 +44,11 @@ public:
 	 */
 	static TSharedRef<ISequencerTrackEditor> CreateTrackEditor( TSharedRef<ISequencer> OwningSequencer );
 
-	//~ ISequencerTrackEditor Interface
-
-	virtual TSharedRef<ISequencerSection> MakeSectionInterface(UMovieSceneSection& SectionObject, UMovieSceneTrack& Track, FGuid ObjectBinding) override;
-
 protected:
 
 	//~ FPropertyTrackEditor Interface
 
-	virtual void GenerateKeysFromPropertyChanged( const FPropertyChangedParams& PropertyChangedParams, TArray<FMarginKey>& NewGeneratedKeys, TArray<FMarginKey>& DefaultGeneratedKeys) override;
+	virtual void GenerateKeysFromPropertyChanged( const FPropertyChangedParams& PropertyChangedParams, FGeneratedTrackKeys& OutGeneratedKeys) override;
 
 private:
 
diff --git a/Engine/Source/Editor/UMGEditor/Private/Animation/Sequencer2DTransformTrackEditor.cpp b/Engine/Source/Editor/UMGEditor/Private/Animation/Sequencer2DTransformTrackEditor.cpp
--- a/Engine/Source/Editor/UMGEditor/Private/Animation/Sequencer2DTransformTrackEditor.cpp
+++ b/Engine/Source/Editor/UMGEditor/Private/Animation/Sequencer2DTransformTrackEditor.cpp
@@ -1,165 +1,93 @@
 // Copyright 1998-2018 Epic Games, Inc. All Rights Reserved.
 
 #include "Animation/Sequencer2DTransformTrackEditor.h"
-#include "PropertySection.h"
+#include "Slate/WidgetTransform.h"
 #include "ISectionLayoutBuilder.h"
-#include "FloatCurveKeyArea.h"
 
 
 FName F2DTransformTrackEditor::TranslationName( "Translation" );
 FName F2DTransformTrackEditor::ScaleName( "Scale" );
 FName F2DTransformTrackEditor::ShearName( "Shear" );
 FName F2DTransformTrackEditor::AngleName( "Angle" );
+FName F2DTransformTrackEditor::ChannelXName( "X" );
+FName F2DTransformTrackEditor::ChannelYName( "Y" );
 
-
-class F2DTransformSection
-	: public FPropertySection
+TSharedRef<ISequencerTrackEditor> F2DTransformTrackEditor::CreateTrackEditor( TSharedRef<ISequencer> InSequencer )
 {
-public:
-
-	F2DTransformSection(ISequencer* InSequencer, FGuid InObjectBinding, FName InPropertyName, const FString& InPropertyPath, UMovieSceneSection& InSectionObject, const FText& SectionName)
-		: FPropertySection(InSequencer, InObjectBinding, InPropertyName, InPropertyPath, InSectionObject, SectionName)
-	{
-	}
-
-	virtual void GenerateSectionLayout( class ISectionLayoutBuilder& LayoutBuilder ) const override
-	{
-		static const FLinearColor GreenKeyAreaColor(0.0f, 0.7f, 0.0f, 0.2f);
-		static const FLinearColor RedKeyAreaColor(0.7f, 0.0f, 0.0f, 0.2f);
-
-		UMovieScene2DTransformSection* TransformSection = Cast<UMovieScene2DTransformSection>(&SectionObject);
-
-		// Translation
-		TAttribute<TOptional<float>> TranslationXExternalValue = TAttribute<TOptional<float>>::Create(
-			TAttribute<TOptional<float>>::FGetter::CreateRaw(this, &F2DTransformSection::GetTranslationXValue));
-		TSharedRef<FFloatCurveKeyArea> TranslationXKeyArea = MakeShareable(new FFloatCurveKeyArea(&TransformSection->GetTranslationCurve(EAxis::X), TranslationXExternalValue, TransformSection, RedKeyAreaColor));
-
-		TAttribute<TOptional<float>> TranslationYExternalValue = TAttribute<TOptional<float>>::Create(
-			TAttribute<TOptional<float>>::FGetter::CreateRaw(this, &F2DTransformSection::GetTranslationYValue));
-		TSharedRef<FFloatCurveKeyArea> TranslationYKeyArea = MakeShareable(new FFloatCurveKeyArea(&TransformSection->GetTranslationCurve(EAxis::Y), TranslationYExternalValue, TransformSection, GreenKeyAreaColor));
-
-		// Rotation
-		TAttribute<TOptional<float>> RotationExternalValue = TAttribute<TOptional<float>>::Create(
-			TAttribute<TOptional<float>>::FGetter::CreateRaw(this, &F2DTransformSection::GetRotationValue));
-		TSharedRef<FFloatCurveKeyArea> RotationKeyArea = MakeShareable(new FFloatCurveKeyArea(&TransformSection->GetRotationCurve(), RotationExternalValue, TransformSection));
-
-		// Scale
-		TAttribute<TOptional<float>> ScaleXExternalValue = TAttribute<TOptional<float>>::Create(
-			TAttribute<TOptional<float>>::FGetter::CreateRaw(this, &F2DTransformSection::GetScaleXValue));
-		TSharedRef<FFloatCurveKeyArea> ScaleXKeyArea = MakeShareable(new FFloatCurveKeyArea(&TransformSection->GetScaleCurve(EAxis::X), ScaleXExternalValue, TransformSection, RedKeyAreaColor));
-
-		TAttribute<TOptional<float>> ScaleYExternalValue = TAttribute<TOptional<float>>::Create(
-			TAttribute<TOptional<float>>::FGetter::CreateRaw(this, &F2DTransformSection::GetScaleYValue));
-		TSharedRef<FFloatCurveKeyArea> ScaleYKeyArea = MakeShareable(new FFloatCurveKeyArea(&TransformSection->GetScaleCurve(EAxis::Y), ScaleYExternalValue, TransformSection, GreenKeyAreaColor));
-
-		// Shear
-		TAttribute<TOptional<float>> ShearXExternalValue = TAttribute<TOptional<float>>::Create(
-			TAttribute<TOptional<float>>::FGetter::CreateRaw(this, &F2DTransformSection::GetShearXValue));
-		TSharedRef<FFloatCurveKeyArea> ShearXKeyArea = MakeShareable(new FFloatCurveKeyArea(&TransformSection->GetShearCurve(EAxis::X), ShearXExternalValue, TransformSection, RedKeyAreaColor));
-
-		TAttribute<TOptional<float>> ShearYExternalValue = TAttribute<TOptional<float>>::Create(
-			TAttribute<TOptional<float>>::FGetter::CreateRaw(this, &F2DTransformSection::GetShearYValue));
-		TSharedRef<FFloatCurveKeyArea> ShearYKeyArea = MakeShareable(new FFloatCurveKeyArea(&TransformSection->GetShearCurve(EAxis::Y), ShearYExternalValue, TransformSection, GreenKeyAreaColor));
-
-		// This generates the tree structure for the transform section
-		LayoutBuilder.PushCategory("Translation", NSLOCTEXT("F2DTransformSection", "TranslationArea", "Translation"));
-			LayoutBuilder.AddKeyArea("Translation.X", NSLOCTEXT("F2DTransformSection", "TransXArea", "X"), TranslationXKeyArea);
-			LayoutBuilder.AddKeyArea("Translation.Y", NSLOCTEXT("F2DTransformSection", "TransYArea", "Y"), TranslationYKeyArea);
-		LayoutBuilder.PopCategory();
-
-		LayoutBuilder.PushCategory("Rotation", NSLOCTEXT("F2DTransformSection", "RotationArea", "Rotation"));
-			LayoutBuilder.AddKeyArea("Rotation.Angle", NSLOCTEXT("F2DTransformSection", "AngleArea", "Angle"), RotationKeyArea);
-		LayoutBuilder.PopCategory();
-
-		LayoutBuilder.PushCategory("Scale", NSLOCTEXT("F2DTransformSection", "ScaleArea", "Scale"));
-			LayoutBuilder.AddKeyArea("Scale.X", NSLOCTEXT("F2DTransformSection", "ScaleXArea", "X"), ScaleXKeyArea);
-			LayoutBuilder.AddKeyArea("Scale.Y", NSLOCTEXT("F2DTransformSection", "ScaleYArea", "Y"), ScaleYKeyArea);
-		LayoutBuilder.PopCategory();
-
-		LayoutBuilder.PushCategory("Shear", NSLOCTEXT("F2DTransformSection", "ShearArea", "Shear"));
-			LayoutBuilder.AddKeyArea("Shear.X", NSLOCTEXT("F2DTransformSection", "SheerXArea", "X"), ShearXKeyArea);
-			LayoutBuilder.AddKeyArea("Shear.Y", NSLOCTEXT("F2DTransformSection", "SheerYArea", "Y"), ShearYKeyArea);
-		LayoutBuilder.PopCategory();
-	}
-
-private:
-
-	TOptional<float> GetTranslationXValue() const
-	{
-		TOptional<FWidgetTransform> Transform = GetPropertyValue<FWidgetTransform>();
-		return Transform.IsSet() ? TOptional<float>(Transform.GetValue().Translation.X) : TOptional<float>();
-	}
-
-	TOptional<float> GetTranslationYValue() const
-	{
-		TOptional<FWidgetTransform> Transform = GetPropertyValue<FWidgetTransform>();
-		return Transform.IsSet() ? TOptional<float>(Transform.GetValue().Translation.Y) : TOptional<float>();
-	}
+	return MakeShareable( new F2DTransformTrackEditor( InSequencer ) );
+}
 
-	TOptional<float> GetRotationValue() const
-	{
-		TOptional<FWidgetTransform> Transform = GetPropertyValue<FWidgetTransform>();
-		return Transform.IsSet() ? TOptional<float>(Transform.GetValue().Angle) : TOptional<float>();
-	}
 
-	TOptional<float> GetScaleXValue() const
-	{
-		TOptional<FWidgetTransform> Transform = GetPropertyValue<FWidgetTransform>();
-		return Transform.IsSet() ? TOptional<float>(Transform.GetValue().Scale.X) : TOptional<float>();
-	}
+void F2DTransformTrackEditor::GenerateKeysFromPropertyChanged( const FPropertyChangedParams& PropertyChangedParams, FGeneratedTrackKeys& OutGeneratedKeys )
+{
+	FPropertyPath StructPath = PropertyChangedParams.StructPathToKey;
 
-	TOptional<float> GetScaleYValue() const
-	{
-		TOptional<FWidgetTransform> Transform = GetPropertyValue<FWidgetTransform>();
-		return Transform.IsSet() ? TOptional<float>(Transform.GetValue().Scale.Y) : TOptional<float>();
-	}
+	bool bKeyTranslationX = false;
+	bool bKeyTranslationY = false;
+	bool bKeyAngle = false;
+	bool bKeyScaleX = false;
+	bool bKeyScaleY = false;
+	bool bKeyShearX = false;
+	bool bKeyShearY = false;
 
-	TOptional<float> GetShearXValue() const
+	if (StructPath.GetNumProperties() == 0)
 	{
-		TOptional<FWidgetTransform> Transform = GetPropertyValue<FWidgetTransform>();
-		return Transform.IsSet() ? TOptional<float>(Transform.GetValue().Shear.X) : TOptional<float>();
+		bKeyTranslationX = bKeyTranslationY = bKeyAngle = bKeyScaleX = bKeyScaleY = bKeyShearX = bKeyShearY = true;
 	}
-
-	TOptional<float> GetShearYValue() const
+	else
 	{
-		TOptional<FWidgetTransform> Transform = GetPropertyValue<FWidgetTransform>();
-		return Transform.IsSet() ? TOptional<float>(Transform.GetValue().Shear.Y) : TOptional<float>();
+		if (StructPath.GetRootProperty().Property->GetFName() == TranslationName)
+		{
+			if (StructPath.GetLeafMostProperty() != StructPath.GetRootProperty())
+			{
+				bKeyTranslationX = StructPath.GetLeafMostProperty().Property->GetFName() == ChannelXName;
+				bKeyTranslationY = StructPath.GetLeafMostProperty().Property->GetFName() == ChannelYName;
+			}
+			else
+			{
+				bKeyTranslationX = bKeyTranslationY = true;
+			}
+		}
+
+		if (StructPath.GetRootProperty().Property->GetFName() == AngleName)
+		{
+			bKeyAngle = true;
+		}
+
+		if (StructPath.GetRootProperty().Property->GetFName() == ScaleName)
+		{
+			if (StructPath.GetLeafMostProperty() != StructPath.GetRootProperty())
+			{
+				bKeyScaleX = StructPath.GetLeafMostProperty().Property->GetFName() == ChannelXName;
+				bKeyScaleY = StructPath.GetLeafMostProperty().Property->GetFName() == ChannelYName;
+			}
+			else
+			{
+				bKeyScaleX = bKeyScaleY = true;
+			}
+		}
+
+		if (StructPath.GetRootProperty().Property->GetFName() == ShearName)
+		{
+			if (StructPath.GetLeafMostProperty() != StructPath.GetRootProperty())
+			{
+				bKeyShearX = StructPath.GetLeafMostProperty().Property->GetFName() == ChannelXName;
+				bKeyShearY = StructPath.GetLeafMostProperty().Property->GetFName() == ChannelYName;
+			}
+			else
+			{
+				bKeyShearX = bKeyShearY = true;
+			}
+		}
 	}
 
-};
-
-
-TSharedRef<ISequencerTrackEditor> F2DTransformTrackEditor::CreateTrackEditor( TSharedRef<ISequencer> InSequencer )
-{
-	return MakeShareable( new F2DTransformTrackEditor( InSequencer ) );
-}
-
-
-TSharedRef<ISequencerSection> F2DTransformTrackEditor::MakeSectionInterface(UMovieSceneSection& SectionObject, UMovieSceneTrack& Track, FGuid ObjectBinding)
-{
-	UMovieScenePropertyTrack* PropertyTrack = Cast<UMovieScenePropertyTrack>(&Track);
-	checkf(PropertyTrack != nullptr, TEXT("Incompatible track in F2DTransformTrackEditor"));
-	return MakeShareable(new F2DTransformSection(GetSequencer().Get(), ObjectBinding, PropertyTrack->GetPropertyName(), PropertyTrack->GetPropertyPath(), SectionObject, Track.GetDisplayName()));
-}
-
-
-void F2DTransformTrackEditor::GenerateKeysFromPropertyChanged( const FPropertyChangedParams& PropertyChangedParams, TArray<F2DTransformKey>& NewGeneratedKeys, TArray<F2DTransformKey>& DefaultGeneratedKeys )
-{
-	FName ChannelName = PropertyChangedParams.StructPropertyNameToKey;
 	FWidgetTransform Transform = PropertyChangedParams.GetPropertyValue<FWidgetTransform>();
 
-	TArray<F2DTransformKey>& TranslationKeys = ChannelName == NAME_None || ChannelName == TranslationName ? NewGeneratedKeys : DefaultGeneratedKeys;
-	TranslationKeys.Add( F2DTransformKey( EKey2DTransformChannel::Translation, EKey2DTransformAxis::X, Transform.Translation.X ) );
-	TranslationKeys.Add( F2DTransformKey( EKey2DTransformChannel::Translation, EKey2DTransformAxis::Y, Transform.Translation.Y ) );
-
-	TArray<F2DTransformKey>& ScaleKeys = ChannelName == NAME_None || ChannelName == ScaleName ? NewGeneratedKeys : DefaultGeneratedKeys;
-	ScaleKeys.Add( F2DTransformKey( EKey2DTransformChannel::Scale, EKey2DTransformAxis::X, Transform.Scale.X ) );
-	ScaleKeys.Add( F2DTransformKey( EKey2DTransformChannel::Scale, EKey2DTransformAxis::Y, Transform.Scale.Y ) );
-
-	TArray<F2DTransformKey>& ShearKeys = ChannelName == NAME_None || ChannelName == ShearName ? NewGeneratedKeys : DefaultGeneratedKeys;
-	ShearKeys.Add( F2DTransformKey( EKey2DTransformChannel::Shear, EKey2DTransformAxis::X, Transform.Shear.X ) );
-	ShearKeys.Add( F2DTransformKey( EKey2DTransformChannel::Shear, EKey2DTransformAxis::Y, Transform.Shear.Y ) );
-
-	TArray<F2DTransformKey>& AngleKeys = ChannelName == NAME_None || ChannelName == AngleName ? NewGeneratedKeys : DefaultGeneratedKeys;
-	AngleKeys.Add( F2DTransformKey( EKey2DTransformChannel::Rotation, EKey2DTransformAxis::None, Transform.Angle ) );
-}
+	OutGeneratedKeys.Add(FMovieSceneChannelValueSetter::Create<FMovieSceneFloatChannel>(0, Transform.Translation.X, bKeyTranslationX));
+	OutGeneratedKeys.Add(FMovieSceneChannelValueSetter::Create<FMovieSceneFloatChannel>(1, Transform.Translation.Y, bKeyTranslationY));
+	OutGeneratedKeys.Add(FMovieSceneChannelValueSetter::Create<FMovieSceneFloatChannel>(2, Transform.Angle,         bKeyAngle));
+	OutGeneratedKeys.Add(FMovieSceneChannelValueSetter::Create<FMovieSceneFloatChannel>(3, Transform.Scale.X,       bKeyScaleX));
+	OutGeneratedKeys.Add(FMovieSceneChannelValueSetter::Create<FMovieSceneFloatChannel>(4, Transform.Scale.Y,       bKeyScaleY));
+	OutGeneratedKeys.Add(FMovieSceneChannelValueSetter::Create<FMovieSceneFloatChannel>(5, Transform.Shear.X,       bKeyShearX));
+	OutGeneratedKeys.Add(FMovieSceneChannelValueSetter::Create<FMovieSceneFloatChannel>(6, Transform.Shear.Y,       bKeyShearY));
+}
\ No newline at end of file
diff --git a/Engine/Source/Editor/UMGEditor/Private/Animation/Sequencer2DTransformTrackEditor.h b/Engine/Source/Editor/UMGEditor/Private/Animation/Sequencer2DTransformTrackEditor.h
--- a/Engine/Source/Editor/UMGEditor/Private/Animation/Sequencer2DTransformTrackEditor.h
+++ b/Engine/Source/Editor/UMGEditor/Private/Animation/Sequencer2DTransformTrackEditor.h
@@ -12,7 +12,7 @@
 #include "Animation/MovieScene2DTransformSection.h"
 
 class F2DTransformTrackEditor
-	: public FPropertyTrackEditor<UMovieScene2DTransformTrack, UMovieScene2DTransformSection, F2DTransformKey>
+	: public FPropertyTrackEditor<UMovieScene2DTransformTrack>
 {
 public:
 
@@ -22,7 +22,7 @@ public:
 	 * @param InSequencer The sequencer instance to be used by this tool
 	 */
 	F2DTransformTrackEditor( TSharedRef<ISequencer> InSequencer )
-		: FPropertyTrackEditor<UMovieScene2DTransformTrack, UMovieScene2DTransformSection, F2DTransformKey>( InSequencer, GetAnimatedPropertyTypes() )
+		: FPropertyTrackEditor( InSequencer, GetAnimatedPropertyTypes() )
 	{
 	}
 
@@ -42,19 +42,17 @@ public:
 	 */
 	static TSharedRef<ISequencerTrackEditor> CreateTrackEditor( TSharedRef<ISequencer> OwningSequencer );
 
-	//~ ISequencerTrackEditor Interface
-
-	virtual TSharedRef<ISequencerSection> MakeSectionInterface(UMovieSceneSection& SectionObject, UMovieSceneTrack& Track, FGuid ObjectBinding) override;
-
 protected:
 
 	// FPropertyTrackEditor interface
 
-	virtual void GenerateKeysFromPropertyChanged( const FPropertyChangedParams& PropertyChangedParams, TArray<F2DTransformKey>& NewGeneratedKeys, TArray<F2DTransformKey>& DefaultGeneratedKeys ) override;
+	virtual void GenerateKeysFromPropertyChanged( const FPropertyChangedParams& PropertyChangedParams, FGeneratedTrackKeys& OutGeneratedKeys ) override;
 
 private:
 	static FName TranslationName;
 	static FName ScaleName;
 	static FName ShearName;
 	static FName AngleName;
+	static FName ChannelXName;
+	static FName ChannelYName;
 };
diff --git a/Engine/Source/Editor/UMGEditor/Private/Animation/UMGDetailKeyframeHandler.cpp b/Engine/Source/Editor/UMGEditor/Private/Animation/UMGDetailKeyframeHandler.cpp
--- a/Engine/Source/Editor/UMGEditor/Private/Animation/UMGDetailKeyframeHandler.cpp
+++ b/Engine/Source/Editor/UMGEditor/Private/Animation/UMGDetailKeyframeHandler.cpp
@@ -4,6 +4,7 @@
 #include "Animation/WidgetAnimation.h"
 
 #include "PropertyHandle.h"
+#include "MovieScene.h"
 
 FUMGDetailKeyframeHandler::FUMGDetailKeyframeHandler(TSharedPtr<FWidgetBlueprintEditor> InBlueprintEditor)
 	: BlueprintEditor( InBlueprintEditor )
@@ -20,6 +21,26 @@ bool FUMGDetailKeyframeHandler::IsPropertyKeyingEnabled() const
 	return Sequence != nullptr && Sequence != UWidgetAnimation::GetNullAnimation();
 }
 
+bool FUMGDetailKeyframeHandler::IsPropertyAnimated(const IPropertyHandle& PropertyHandle, UObject *ParentObject) const
+{
+	TSharedPtr<ISequencer> Sequencer = BlueprintEditor.Pin()->GetSequencer();
+	if (Sequencer.IsValid() && Sequencer->GetFocusedMovieSceneSequence())
+	{
+		FGuid ObjectHandle = Sequencer->GetHandleToObject(ParentObject);
+		if (ObjectHandle.IsValid()) 
+		{
+			UMovieScene* MovieScene = Sequencer->GetFocusedMovieSceneSequence()->GetMovieScene();
+			UProperty* Property = PropertyHandle.GetProperty();
+			TSharedRef<FPropertyPath> PropertyPath = FPropertyPath::CreateEmpty();
+			PropertyPath->AddProperty(FPropertyInfo(Property));
+			FName PropertyName(*PropertyPath->ToString(TEXT(".")));
+			TSubclassOf<UMovieSceneTrack> TrackClass; //use empty @todo find way to get the UMovieSceneTrack from the Property type.
+			return MovieScene->FindTrack(TrackClass, ObjectHandle, PropertyName) != nullptr;
+		}
+	}
+	return false;
+}
+
 void FUMGDetailKeyframeHandler::OnKeyPropertyClicked(const IPropertyHandle& KeyedPropertyHandle)
 {
 	TArray<UObject*> Objects;
diff --git a/Engine/Source/Editor/UMGEditor/Private/Animation/UMGDetailKeyframeHandler.h b/Engine/Source/Editor/UMGEditor/Private/Animation/UMGDetailKeyframeHandler.h
--- a/Engine/Source/Editor/UMGEditor/Private/Animation/UMGDetailKeyframeHandler.h
+++ b/Engine/Source/Editor/UMGEditor/Private/Animation/UMGDetailKeyframeHandler.h
@@ -19,6 +19,8 @@ public:
 
 	virtual void OnKeyPropertyClicked(const IPropertyHandle& KeyedPropertyHandle) override;
 
+	virtual bool IsPropertyAnimated(const class IPropertyHandle& PropertyHandle, UObject *ParentObject) const override;
+
 private:
 	TWeakPtr<class FWidgetBlueprintEditor> BlueprintEditor;
 };
diff --git a/Engine/Source/Editor/UMGEditor/Private/TabFactory/AnimationTabSummoner.cpp b/Engine/Source/Editor/UMGEditor/Private/TabFactory/AnimationTabSummoner.cpp
--- a/Engine/Source/Editor/UMGEditor/Private/TabFactory/AnimationTabSummoner.cpp
+++ b/Engine/Source/Editor/UMGEditor/Private/TabFactory/AnimationTabSummoner.cpp
@@ -402,8 +402,14 @@ private:
 
 		NewAnimation->MovieScene = NewObject<UMovieScene>(NewAnimation, NewAnimation->GetFName(), RF_Transactional);
 
-		NewAnimation->MovieScene->SetPlaybackRange(InTime, OutTime);
-		NewAnimation->MovieScene->GetEditorData().WorkingRange = TRange<float>(InTime, OutTime);
+		// Default to 20 fps display rate (as was the previous default in USequencerSettings)
+		NewAnimation->MovieScene->SetPlaybackFrameRate(FFrameRate(20, 1));
+
+		const FFrameTime InFrame  = InTime  * NewAnimation->MovieScene->GetFrameResolution();
+		const FFrameTime OutFrame = OutTime * NewAnimation->MovieScene->GetFrameResolution();
+		NewAnimation->MovieScene->SetPlaybackRange(TRange<FFrameNumber>(InFrame.FrameNumber, OutFrame.FrameNumber+1));
+		NewAnimation->MovieScene->GetEditorData().WorkStart = InTime;
+		NewAnimation->MovieScene->GetEditorData().WorkEnd   = OutTime;
 
 		bool bRequestRename = true;
 		bool bNewAnimation = true;
diff --git a/Engine/Source/Editor/UMGEditor/Private/WidgetBlueprintEditor.cpp b/Engine/Source/Editor/UMGEditor/Private/WidgetBlueprintEditor.cpp
--- a/Engine/Source/Editor/UMGEditor/Private/WidgetBlueprintEditor.cpp
+++ b/Engine/Source/Editor/UMGEditor/Private/WidgetBlueprintEditor.cpp
@@ -65,12 +65,6 @@ FWidgetBlueprintEditor::FWidgetBlueprintEditor()
 			FAssetEditorExtender::CreateRaw(this, &FWidgetBlueprintEditor::GetAddTrackSequencerExtender));
 		SequencerAddTrackExtenderHandle = SequencerModule.GetAddTrackMenuExtensibilityManager()->GetExtenderDelegates()[NewIndex].GetHandle();
 	}
-
-	{
-		int32 NewIndex = SequencerModule.GetObjectBindingContextMenuExtensibilityManager()->GetExtenderDelegates().Add(
-			FAssetEditorExtender::CreateRaw(this, &FWidgetBlueprintEditor::GetObjectBindingContextMenuExtender));
-		SequencerObjectBindingExtenderHandle = SequencerModule.GetObjectBindingContextMenuExtensibilityManager()->GetExtenderDelegates()[NewIndex].GetHandle();
-	}
 }
 
 FWidgetBlueprintEditor::~FWidgetBlueprintEditor()
@@ -98,10 +92,6 @@ FWidgetBlueprintEditor::~FWidgetBlueprintEditor()
 		return SequencerAddTrackExtenderHandle == Extender.GetHandle();
 	});
 
-	SequencerModule.GetObjectBindingContextMenuExtensibilityManager()->GetExtenderDelegates().RemoveAll([this](const FAssetEditorExtender& Extender)
-	{
-		return SequencerObjectBindingExtenderHandle == Extender.GetHandle();
-	});
 }
 
 void FWidgetBlueprintEditor::InitWidgetBlueprintEditor(const EToolkitMode::Type Mode, const TSharedPtr< IToolkitHost >& InitToolkitHost, const TArray<UBlueprint*>& InBlueprints, bool bShouldOpenInDefaultsMode)
@@ -678,20 +668,25 @@ TSharedPtr<ISequencer>& FWidgetBlueprintEditor::GetSequencer()
 {
 	if(!Sequencer.IsValid())
 	{
-		const float InTime = 0.f;
+		const float InTime  = 0.f;
 		const float OutTime = 5.0f;
 
 		FSequencerViewParams ViewParams(TEXT("UMGSequencerSettings"));
 		{
-			ViewParams.InitialScrubPosition = 0;
 			ViewParams.OnGetAddMenuContent = FOnGetAddMenuContent::CreateSP(this, &FWidgetBlueprintEditor::OnGetAnimationAddMenuContent);
+		    ViewParams.OnBuildCustomContextMenuForGuid = FOnBuildCustomContextMenuForGuid::CreateSP(this, &FWidgetBlueprintEditor::OnBuildCustomContextMenuForGuid);
 		}
 
 		FSequencerInitParams SequencerInitParams;
 		{
 			UWidgetAnimation* NullAnimation = UWidgetAnimation::GetNullAnimation();
-			NullAnimation->MovieScene->SetPlaybackRange(InTime, OutTime);
-			NullAnimation->MovieScene->GetEditorData().WorkingRange = TRange<float>(InTime, OutTime);
+			FFrameRate FrameResolution = NullAnimation->MovieScene->GetFrameResolution();
+			FFrameNumber StartFrame = (InTime  * FrameResolution).FloorToFrame();
+			FFrameNumber EndFrame   = (OutTime * FrameResolution).CeilToFrame();
+			NullAnimation->MovieScene->SetPlaybackRange(StartFrame, (EndFrame-StartFrame).Value);
+			FMovieSceneEditorData& EditorData = NullAnimation->MovieScene->GetEditorData();
+			EditorData.WorkStart = InTime;
+			EditorData.WorkEnd   = OutTime;
 
 			SequencerInitParams.ViewParams = ViewParams;
 			SequencerInitParams.RootSequence = NullAnimation;
@@ -1103,16 +1098,40 @@ TSharedRef<FExtender> FWidgetBlueprintEditor::GetAddTrackSequencerExtender( cons
 	return AddTrackMenuExtender;
 }
 
-TSharedRef<FExtender> FWidgetBlueprintEditor::GetObjectBindingContextMenuExtender(const TSharedRef<FUICommandList> CommandList, const TArray<UObject*> ContextSensitiveObjects)
+void FWidgetBlueprintEditor::OnBuildCustomContextMenuForGuid(FMenuBuilder& MenuBuilder, FGuid ObjectBinding)
 {
-	TSharedRef<FExtender> ObjectBindingMenuExtender(new FExtender());
+	if (CurrentAnimation.IsValid())
+	{
+		FWidgetReference SelectedWidget;
+		if (SelectedWidgets.Num() == 1)
+		{
+			for (FWidgetReference Widget : SelectedWidgets)
+			{
+				SelectedWidget = Widget;
+			}
+		}
+		if (SelectedWidget.IsValid())
+		{
+			//need to make sure it's a widget, if not bound assume it is.
+			UWidget* BoundWidget = nullptr;
+			bool bNotBound = true;
+			for (TWeakObjectPtr<> WeakObjectPtr : GetSequencer()->FindObjectsInCurrentSequence(ObjectBinding))
+			{
+				BoundWidget = Cast<UWidget>(WeakObjectPtr.Get());
+				bNotBound = false;
+				break;
+			}
+			if (bNotBound || (BoundWidget && SelectedWidget.GetPreview()->GetTypedOuter<UWidgetTree>() == BoundWidget->GetTypedOuter<UWidgetTree>()) )
+			{
+				FUIAction ReplaceWithMenuAction(FExecuteAction::CreateRaw(this, &FWidgetBlueprintEditor::ReplaceTrackWithSelectedWidget, SelectedWidget, (UWidget*)nullptr, ObjectBinding));
 
-	ObjectBindingMenuExtender->AddMenuExtension(
-		"Edit",
-		EExtensionHook::First,
-		CommandList,
-		FMenuExtensionDelegate::CreateRaw(this, &FWidgetBlueprintEditor::ExtendSequencerObjectBindingMenu, ContextSensitiveObjects));
-	return ObjectBindingMenuExtender;
+				FText ReplaceWithLabel = FText::Format(LOCTEXT("ReplaceObject", "Replace with {0}"), FText::FromString(SelectedWidget.GetPreview()->GetName()));
+				FText ReplaceWithToolTip = FText::Format(LOCTEXT("ReplaceObjectToolTip", "Replace the bound widget in this animation with {0}"), FText::FromString(SelectedWidget.GetPreview()->GetName()));
+				MenuBuilder.AddMenuSeparator();
+				MenuBuilder.AddMenuEntry(ReplaceWithLabel, ReplaceWithToolTip, FSlateIcon(), ReplaceWithMenuAction);
+			}
+		}
+	}
 }
 
 void FWidgetBlueprintEditor::ExtendSequencerAddTrackMenu( FMenuBuilder& AddTrackMenuBuilder, const TArray<UObject*> ContextObjects )
@@ -1162,7 +1181,7 @@ void FWidgetBlueprintEditor::ExtendSequencerAddTrackMenu( FMenuBuilder& AddTrack
 	}
 }
 
-void FWidgetBlueprintEditor::ReplaceTrackWithSelectedWidget(FWidgetReference SelectedWidget, UWidget* BoundWidget)
+void FWidgetBlueprintEditor::ReplaceTrackWithSelectedWidget(FWidgetReference SelectedWidget, UWidget* BoundWidget, FGuid ObjectId)
 {
 	const FScopedTransaction Transaction( LOCTEXT( "ReplaceTrackWithSelectedWidget", "Replace Track with Selected Widget" ) );
 
@@ -1170,8 +1189,6 @@ void FWidgetBlueprintEditor::ReplaceTrackWithSelectedWidget(FWidgetReference Sel
 	UMovieScene* MovieScene = WidgetAnimation->GetMovieScene();
 	UWidget* PreviewWidget = SelectedWidget.GetPreview();
 	UWidget* TemplateWidget = SelectedWidget.GetTemplate();
-	FGuid ObjectId = BoundWidget ? Sequencer->FindObjectId(*BoundWidget, MovieSceneSequenceID::Root) : FGuid();
-
 	// Try find if the SelectedWidget is already bound, if so return
 	FGuid SelectedWidgetId = Sequencer->FindObjectId(*PreviewWidget, MovieSceneSequenceID::Root);
 	if (SelectedWidgetId.IsValid())
@@ -1207,7 +1224,7 @@ void FWidgetBlueprintEditor::ReplaceTrackWithSelectedWidget(FWidgetReference Sel
 						if (!SelectedWidget.GetTemplate()->FindFunction(FunctionName))
 						{
 							// Exists a track that's not compatible 
-							FNotificationInfo Info(LOCTEXT("IncompatibleTrackToReplaceWith", "Selected Widget doesn't match to a Property this track is bound to"));
+							FNotificationInfo Info(FText::Format(LOCTEXT("IncompatibleTrackToReplaceWith", "Selected Widget doesn't match to a Property '{0}' this track is bound to"), FText::FromString(PropertyName)));
 							Info.FadeInDuration = 0.1f;
 							Info.FadeOutDuration = 0.5f;
 							Info.ExpireDuration = 2.5f;
@@ -1260,6 +1277,7 @@ void FWidgetBlueprintEditor::ReplaceTrackWithSelectedWidget(FWidgetReference Sel
 	Sequencer->NotifyMovieSceneDataChanged(EMovieSceneDataChangeType::MovieSceneStructureItemsChanged);
 }
 
+
 void FWidgetBlueprintEditor::ExtendSequencerObjectBindingMenu(FMenuBuilder& ObjectBindingMenuBuilder, const TArray<UObject*> ContextObjects)
 {
 	FWidgetReference SelectedWidget;
@@ -1275,7 +1293,8 @@ void FWidgetBlueprintEditor::ExtendSequencerObjectBindingMenu(FMenuBuilder& Obje
 		UWidget* BoundWidget = Cast<UWidget>(ContextObjects[0]);
 		if (BoundWidget && SelectedWidget.GetPreview()->GetTypedOuter<UWidgetTree>() == BoundWidget->GetTypedOuter<UWidgetTree>() )
 		{
-			FUIAction ReplaceWithMenuAction(FExecuteAction::CreateRaw(this, &FWidgetBlueprintEditor::ReplaceTrackWithSelectedWidget, SelectedWidget, BoundWidget));
+			FGuid ObjectId = Sequencer->FindObjectId(*BoundWidget, MovieSceneSequenceID::Root);
+			FUIAction ReplaceWithMenuAction(FExecuteAction::CreateRaw(this, &FWidgetBlueprintEditor::ReplaceTrackWithSelectedWidget, SelectedWidget, BoundWidget, ObjectId));
 
 			FText ReplaceWithLabel = FText::Format(LOCTEXT("ReplaceObject", "Replace with {0}"), FText::FromString(SelectedWidget.GetPreview()->GetName()));
 			FText ReplaceWithToolTip = FText::Format(LOCTEXT("ReplaceObjectToolTip", "Replace the bound widget in this animation with {0}"), FText::FromString(SelectedWidget.GetPreview()->GetName()));
diff --git a/Engine/Source/Editor/UMGEditor/Private/WidgetBlueprintEditor.h b/Engine/Source/Editor/UMGEditor/Private/WidgetBlueprintEditor.h
--- a/Engine/Source/Editor/UMGEditor/Private/WidgetBlueprintEditor.h
+++ b/Engine/Source/Editor/UMGEditor/Private/WidgetBlueprintEditor.h
@@ -234,6 +234,10 @@ private:
 	/** Populates the sequencer add menu. */
 	void OnGetAnimationAddMenuContent(FMenuBuilder& MenuBuilder, TSharedRef<ISequencer> Sequencer);
 
+	/** Populates sequencer menu when added with objects. Used to handle case where widget is deleted so it's 
+	    Object * is null.*/
+	void OnBuildCustomContextMenuForGuid(FMenuBuilder& MenuBuilder,FGuid ObjectBinding);
+
 	/** Populates the sequencer add submenu for the big list of widgets. */
 	void OnGetAnimationAddMenuContentAllWidgets(FMenuBuilder& MenuBuilder);
 
@@ -243,13 +247,11 @@ private:
 	/** Gets the extender to use for sequencers context sensitive menus and toolbars. */
 	TSharedRef<FExtender> GetAddTrackSequencerExtender(const TSharedRef<FUICommandList> CommandList, const TArray<UObject*> ContextSensitiveObjects);
 
-	TSharedRef<FExtender> GetObjectBindingContextMenuExtender(const TSharedRef<FUICommandList> CommandList, const TArray<UObject*> ContextSensitiveObjects);
-
 	/** Extends the sequencer add track menu. */
 	void ExtendSequencerAddTrackMenu( FMenuBuilder& AddTrackMenuBuilder, const TArray<UObject*> ContextObjects );
 
 	/** Replace track with selected widget function */
-	void ReplaceTrackWithSelectedWidget(FWidgetReference SelectedWidget, UWidget* BoundWidget);
+	void ReplaceTrackWithSelectedWidget(FWidgetReference SelectedWidget, UWidget* BoundWidget, FGuid ObjectBinding);
 
 	/** Extends the sequencer add track menu. */
 	void ExtendSequencerObjectBindingMenu(FMenuBuilder& ObjectBindingMenuBuilder, const TArray<UObject*> ContextObjects);
@@ -343,8 +345,6 @@ private:
 
 	FDelegateHandle SequencerAddTrackExtenderHandle;
 
-	FDelegateHandle SequencerObjectBindingExtenderHandle;
-
 	/** Messages we want to append to the compiler results. */
 	TArray< TSharedRef<class FTokenizedMessage> > DesignerCompilerMessages;
 
diff --git a/Engine/Source/Editor/UnrealEd/Classes/Editor/EditorEngine.h b/Engine/Source/Editor/UnrealEd/Classes/Editor/EditorEngine.h
--- a/Engine/Source/Editor/UnrealEd/Classes/Editor/EditorEngine.h
+++ b/Engine/Source/Editor/UnrealEd/Classes/Editor/EditorEngine.h
@@ -711,12 +711,6 @@ public:
 	DECLARE_EVENT_OneParam( UEditorEngine, FOnEndTransformCamera, UObject& );
 	FOnEndTransformCamera& OnEndCameraMovement() { return OnEndCameraTransformEvent; }
 
-	/** Delegate broadcast by the engine every tick when PIE/SIE is active, to check to see whether we need to
-		be able to capture state for simulating actor (for Sequencer recording features).  The single bool parameter
-		should be set to true if recording features are needed. */
-	DECLARE_EVENT_OneParam( UEditorEngine, FGetActorRecordingState, bool& /* bIsRecordingActive */ );
-	FGetActorRecordingState& GetActorRecordingState() { return GetActorRecordingStateEvent; }
-
 	/** Editor-only event triggered when a HLOD Actor is moved between clusters */
 	DECLARE_EVENT_TwoParams(UEngine, FHLODActorMovedEvent, const AActor*, const AActor*);
 	FHLODActorMovedEvent& OnHLODActorMoved() { return HLODActorMovedEvent; }
@@ -2805,10 +2799,6 @@ private:
 	/** Broadcasts after an Actor is removed from a cluster */
 	FHLODActorRemovedFromClusterEvent HLODActorRemovedFromClusterEvent;
 
-	/** Delegate broadcast by the engine every tick when PIE/SIE is active, to check to see whether we need to
-		be able to capture state for simulating actor (for Sequencer recording features) */
-	FGetActorRecordingState GetActorRecordingStateEvent;
-
 	/** Delegate to be called when a matinee is requested to be opened */
 	FShouldOpenMatineeCallback ShouldOpenMatineeCallback;
 
diff --git a/Engine/Source/Editor/UnrealEd/Private/EditorEngine.cpp b/Engine/Source/Editor/UnrealEd/Private/EditorEngine.cpp
--- a/Engine/Source/Editor/UnrealEd/Private/EditorEngine.cpp
+++ b/Engine/Source/Editor/UnrealEd/Private/EditorEngine.cpp
@@ -1661,32 +1661,6 @@ void UEditorEngine::Tick( float DeltaSeconds, bool bIdleMode )
 
 				FKismetDebugUtilities::NotifyDebuggerOfStartOfGameFrame(PieContext.World());
 
-				static TArray< TWeakObjectPtr<AActor> > RecordedActors;
-				RecordedActors.Reset();
-
-				// Check to see if we want to use sequencer's live recording feature
-				bool bIsRecordingActive = false;
-				GetActorRecordingStateEvent.Broadcast( bIsRecordingActive );
-				if( bIsRecordingActive )
-				{
-					// @todo sequencer livecapture: How do we capture the destruction of actors? (needs to hide the spawned puppet actor, or destroy it)
-					// @todo sequencer livecapture: Actor parenting state is not captured or retained on puppets
-					// @todo sequencer livecapture: Needs to capture state besides transforms (animation, audio, property changes, etc.)
-
-					// @todo sequencer livecapture: Hacky test code for sequencer live capture feature
-					for( FActorIterator ActorIter( PlayWorld ); ActorIter; ++ActorIter )
-					{
-						AActor* Actor = *ActorIter;
-
-						// @todo sequencer livecapture: Restrict to certain actor types for now, just for testing
-						if( Actor->IsA(ASkeletalMeshActor::StaticClass()) || (Actor->IsA(AStaticMeshActor::StaticClass()) && Actor->IsRootComponentMovable()) )
-						{
-							GEditor->BroadcastBeginObjectMovement( *Actor );
-							RecordedActors.Add( Actor );
-						}
-					}				
-				}
-
 				// tick the level
 				PieContext.World()->Tick( LEVELTICK_All, DeltaSeconds );
 				bAWorldTicked = true;
@@ -1699,18 +1673,6 @@ void UEditorEngine::Tick( float DeltaSeconds, bool bIdleMode )
 					UReflectionCaptureComponent::UpdateReflectionCaptureContents(PlayWorld);
 				}
 
-				if( bIsRecordingActive )
-				{
-					for( auto RecordedActorIter( RecordedActors.CreateIterator() ); RecordedActorIter; ++RecordedActorIter )
-					{
-						AActor* Actor = RecordedActorIter->Get();
-						if( Actor != NULL )
-						{
-							GEditor->BroadcastEndObjectMovement( *Actor );
-						}
-					}				
-				}
-
 				FKismetDebugUtilities::NotifyDebuggerOfEndOfGameFrame(PieContext.World());
 			}
 
@@ -2767,17 +2729,22 @@ void UEditorEngine::ApplyDeltaToComponent(USceneComponent* InComponent,
 		{
 			if ( bDelta )
 			{
-				const FQuat ActorQ = InComponent->RelativeRotation.Quaternion();
+				const FRotator Rot = InComponent->RelativeRotation;
+				FRotator ActorRotWind, ActorRotRem;
+				Rot.GetWindingAndRemainder(ActorRotWind, ActorRotRem);
+				const FQuat ActorQ = ActorRotRem.Quaternion();
 				const FQuat DeltaQ = InDeltaRot.Quaternion();
 				const FQuat ResultQ = DeltaQ * ActorQ;
 
-				const FRotator NewActorRot = FRotator( ResultQ );
-
-				InComponent->SetRelativeRotation(NewActorRot);
+				FRotator NewActorRotRem = FRotator(ResultQ);
+				ActorRotRem.SetClosestToMe(NewActorRotRem);
+				FRotator DeltaRot = NewActorRotRem - ActorRotRem;
+				DeltaRot.Normalize();
+				InComponent->SetRelativeRotationExact(Rot + DeltaRot);
 			}
 			else
 			{
-				InComponent->SetRelativeRotation( InDeltaRot );
+				InComponent->SetRelativeRotationExact( InDeltaRot );
 			}
 
 			if ( bDelta )
