Sequencer: Changes to allow setting min/max values outside of current range. For example, if the current range is 0-100 and the minimum is set to 200, the resulting range is 200-300. This applies to playback, view, and working ranges. Existing rules still apply. Setting playback range, can adjust view and working ranges if necessary. Setting view range beyond working range will adjust working range. Setting view and working ranges will not adjust playback range.

diff --git a/Engine/Source/Editor/Sequencer/Private/SSequencer.cpp b/Engine/Source/Editor/Sequencer/Private/SSequencer.cpp
--- a/Engine/Source/Editor/Sequencer/Private/SSequencer.cpp
+++ b/Engine/Source/Editor/Sequencer/Private/SSequencer.cpp
@@ -2015,18 +2015,16 @@ TSharedRef<SWidget> SSequencer::MakePlaybackMenu()
 		// Menu entry for the start position
 		auto OnStartChanged = [=](double NewValue){
 
-			// We clamp the new value when the value is set. We can't clamp in the UI because we need an unset Min/Max for linear scaling to work.
-			double Min = -FLT_MAX;
-			double Max = SequencerPtr.Pin()->GetPlaybackRange().GetUpperBoundValue().Value;
-
-			NewValue = FMath::Clamp(NewValue, Min, Max);
 			FFrameNumber ValueAsFrame = FFrameTime::FromDecimal(NewValue).GetFrame();
+			FFrameNumber PlayStart = ValueAsFrame;
+			FFrameNumber PlayEnd = MovieScene::DiscreteExclusiveUpper(SequencerPtr.Pin()->GetPlaybackRange());
+			if (PlayStart >= PlayEnd)
+			{
+				FFrameNumber Duration = PlayEnd - MovieScene::DiscreteInclusiveLower(SequencerPtr.Pin()->GetPlaybackRange());
+				PlayEnd = PlayStart + Duration;
+			}
 
-			FFrameNumber Upper = MovieScene::DiscreteExclusiveUpper(SequencerPtr.Pin()->GetPlaybackRange());
-
-			TRange<FFrameNumber> NewRange = TRange<FFrameNumber>(FMath::Min(ValueAsFrame, Upper - 1), Upper);
-
-			SequencerPtr.Pin()->SetPlaybackRange(NewRange);
+			SequencerPtr.Pin()->SetPlaybackRange(TRange<FFrameNumber>(PlayStart, PlayEnd));
 
 			TRange<double> PlayRangeSeconds = SequencerPtr.Pin()->GetPlaybackRange() / SequencerPtr.Pin()->GetFocusedTickResolution();
 			const double AdditionalRange = (PlayRangeSeconds.GetUpperBoundValue() - PlayRangeSeconds.GetLowerBoundValue()) * 0.1;
@@ -2078,15 +2076,16 @@ TSharedRef<SWidget> SSequencer::MakePlaybackMenu()
 		// Menu entry for the end position
 		auto OnEndChanged = [=](double NewValue) {
 
-			// We clamp the new value when the value is set. We can't clamp in the UI because we need an unset Min/Max for linear scaling to work.
-			double Min = SequencerPtr.Pin()->GetPlaybackRange().GetLowerBoundValue().Value;
-			double Max = FLT_MAX;
-
-			NewValue = FMath::Clamp(NewValue, Min, Max);
 			FFrameNumber ValueAsFrame = FFrameTime::FromDecimal(NewValue).GetFrame();
+			FFrameNumber PlayStart = MovieScene::DiscreteInclusiveLower(SequencerPtr.Pin()->GetPlaybackRange());
+			FFrameNumber PlayEnd = ValueAsFrame;
+			if (PlayEnd <= PlayStart)
+			{
+				FFrameNumber Duration = MovieScene::DiscreteExclusiveUpper(SequencerPtr.Pin()->GetPlaybackRange()) - PlayStart;
+				PlayStart = PlayEnd - Duration;
+			}
 
-			FFrameNumber Lower = MovieScene::DiscreteInclusiveLower(SequencerPtr.Pin()->GetPlaybackRange());
-			SequencerPtr.Pin()->SetPlaybackRange(TRange<FFrameNumber>(Lower, FMath::Max(ValueAsFrame, Lower)));
+			SequencerPtr.Pin()->SetPlaybackRange(TRange<FFrameNumber>(PlayStart, PlayEnd));
 
 			TRange<double> PlayRangeSeconds = SequencerPtr.Pin()->GetPlaybackRange() / SequencerPtr.Pin()->GetFocusedTickResolution();
 			const double AdditionalRange = (PlayRangeSeconds.GetUpperBoundValue() - PlayRangeSeconds.GetLowerBoundValue()) * 0.1;
diff --git a/Engine/Source/Editor/SequencerWidgets/Private/STimeRange.cpp b/Engine/Source/Editor/SequencerWidgets/Private/STimeRange.cpp
--- a/Engine/Source/Editor/SequencerWidgets/Private/STimeRange.cpp
+++ b/Engine/Source/Editor/SequencerWidgets/Private/STimeRange.cpp
@@ -26,7 +26,7 @@ void STimeRange::Construct( const STimeRange::FArguments& InArgs, TSharedRef<ITi
 		.OnValueCommitted(this, &STimeRange::OnWorkingStartTimeCommitted)
 		.OnValueChanged(this, &STimeRange::OnWorkingStartTimeChanged)
 		.MinValue(TOptional<double>())
-		.MaxValue(this, &STimeRange::MaxWorkingStartTime)
+		.MaxValue(TOptional<double>())
 		.Style(&FEditorStyle::Get().GetWidgetStyle<FSpinBoxStyle>("Sequencer.HyperlinkSpinBox"))
 		.TypeInterface(NumericTypeInterface)
 		.ClearKeyboardFocusOnCommit(true)
@@ -39,7 +39,7 @@ void STimeRange::Construct( const STimeRange::FArguments& InArgs, TSharedRef<ITi
 		.ToolTipText(LOCTEXT("WorkingRangeEnd", "Working Range End"))
 		.OnValueCommitted( this, &STimeRange::OnWorkingEndTimeCommitted )
 		.OnValueChanged( this, &STimeRange::OnWorkingEndTimeChanged )
-		.MinValue(this, &STimeRange::MinWorkingEndTime)
+		.MinValue(TOptional<double>())
 		.MaxValue(TOptional<double>())
 		.Style(&FEditorStyle::Get().GetWidgetStyle<FSpinBoxStyle>("Sequencer.HyperlinkSpinBox"))
 		.TypeInterface(NumericTypeInterface)
@@ -58,7 +58,7 @@ void STimeRange::Construct( const STimeRange::FArguments& InArgs, TSharedRef<ITi
 		.OnValueCommitted( this, &STimeRange::OnViewStartTimeCommitted )
 		.OnValueChanged( this, &STimeRange::OnViewStartTimeChanged )
 		.MinValue(TOptional<double>())
-		.MaxValue(this, &STimeRange::MaxViewStartTime)
+		.MaxValue(TOptional<double>())
 		.Style(&FEditorStyle::Get().GetWidgetStyle<FSpinBoxStyle>("Sequencer.HyperlinkSpinBox"))
 		.TypeInterface(NumericTypeInterface)
 		.ClearKeyboardFocusOnCommit(true)
@@ -72,7 +72,7 @@ void STimeRange::Construct( const STimeRange::FArguments& InArgs, TSharedRef<ITi
 		.ToolTipText(LOCTEXT("ViewEndTimeTooltip", "View Range End Time"))
 		.OnValueCommitted( this, &STimeRange::OnViewEndTimeCommitted )
 		.OnValueChanged( this, &STimeRange::OnViewEndTimeChanged )
-		.MinValue(this, &STimeRange::MinViewEndTime)
+		.MinValue(TOptional<double>())
 		.MaxValue(TOptional<double>())
 		.Style(&FEditorStyle::Get().GetWidgetStyle<FSpinBoxStyle>("Sequencer.HyperlinkSpinBox"))
 		.TypeInterface(NumericTypeInterface)
@@ -274,46 +274,6 @@ double STimeRange::PlayEndTime() const
 	return UpperBound.Value;
 }
 
-TOptional<double> STimeRange::MaxViewStartTime() const
-{
-	return ViewEndTime();
-}
-
-TOptional<double> STimeRange::MinViewEndTime() const
-{
-	return ViewStartTime();
-}
-
-TOptional<double> STimeRange::MinPlayStartTime() const
-{
-	return WorkingStartTime();
-}
-
-TOptional<double> STimeRange::MaxPlayStartTime() const
-{
-	return PlayEndTime();
-}
-
-TOptional<double> STimeRange::MinPlayEndTime() const
-{
-	return PlayStartTime();
-}
-
-TOptional<double> STimeRange::MaxPlayEndTime() const
-{
-	return WorkingEndTime();
-}
-
-TOptional<double> STimeRange::MaxWorkingStartTime() const
-{
-	return ViewEndTime();
-}
-
-TOptional<double> STimeRange::MinWorkingEndTime() const
-{
-	return ViewStartTime();
-}
-
 void STimeRange::OnWorkingStartTimeCommitted(double NewValue, ETextCommit::Type InTextCommit)
 {
 	OnWorkingStartTimeChanged(NewValue);
@@ -354,7 +314,7 @@ void STimeRange::OnWorkingStartTimeChanged(double NewValue)
 
 	if (Time > TimeSliderController->GetViewRange().GetLowerBoundValue())
 	{
-		TimeSliderController->SetViewRange(Time, TimeSliderController->GetViewRange().GetUpperBoundValue(), EViewRangeInterpolation::Immediate);
+		OnViewStartTimeChanged(NewValue);
 	}
 }
 
@@ -368,7 +328,7 @@ void STimeRange::OnWorkingEndTimeChanged(double NewValue)
 
 	if (Time < TimeSliderController->GetViewRange().GetUpperBoundValue())
 	{
-		TimeSliderController->SetViewRange(TimeSliderController->GetViewRange().GetLowerBoundValue(), Time, EViewRangeInterpolation::Immediate);
+		OnViewEndTimeChanged(NewValue);
 	}
 }
 
@@ -377,65 +337,128 @@ void STimeRange::OnViewStartTimeChanged(double NewValue)
 	FFrameRate TickResolution = TimeSliderController->GetTickResolution();
 	double Time = TickResolution.AsSeconds(FFrameTime::FromDecimal(NewValue));
 
+	double ViewStartTime = TimeSliderController->GetViewRange().GetLowerBoundValue();
+	double ViewEndTime = TimeSliderController->GetViewRange().GetUpperBoundValue();
+
+	double ClampStartTime = TimeSliderController.Get()->GetClampRange().GetLowerBoundValue();
+	double ClampEndTime = TimeSliderController.Get()->GetClampRange().GetUpperBoundValue();
+
+	if (Time >= ViewEndTime)
+	{
+		double ViewDuration = ViewEndTime - ViewStartTime;
+		ViewEndTime = Time + ViewDuration;
+
+		if (ViewEndTime > TimeSliderController.Get()->GetClampRange().GetUpperBoundValue())
+		{
+			TimeSliderController->SetClampRange(TimeSliderController->GetClampRange().GetLowerBoundValue(), ViewEndTime);
+		}
+	}
+
+
 	if (Time < TimeSliderController.Get()->GetClampRange().GetLowerBoundValue())
 	{
 		TimeSliderController->SetClampRange(Time, TimeSliderController->GetClampRange().GetUpperBoundValue());
 	}
 
-	TimeSliderController->SetViewRange(Time, TimeSliderController->GetViewRange().GetUpperBoundValue(), EViewRangeInterpolation::Immediate);
+	TimeSliderController->SetViewRange(Time, ViewEndTime, EViewRangeInterpolation::Immediate);
 }
 
+
 void STimeRange::OnViewEndTimeChanged(double NewValue)
 {
 	FFrameRate TickResolution = TimeSliderController->GetTickResolution();
 	double Time = TickResolution.AsSeconds(FFrameTime::FromDecimal(NewValue));
 
+	double ViewStartTime = TimeSliderController->GetViewRange().GetLowerBoundValue();
+	double ViewEndTime = TimeSliderController->GetViewRange().GetUpperBoundValue();
+
+	if (Time <= ViewStartTime)
+	{
+		double ViewDuration = ViewEndTime - ViewStartTime;
+		ViewStartTime = Time - ViewDuration;
+
+		if (ViewStartTime < TimeSliderController.Get()->GetClampRange().GetLowerBoundValue())
+		{
+			TimeSliderController->SetClampRange(ViewStartTime, TimeSliderController->GetClampRange().GetUpperBoundValue());
+		}
+	}
+
 	if (Time > TimeSliderController->GetClampRange().GetUpperBoundValue())
 	{
 		TimeSliderController->SetClampRange(TimeSliderController->GetClampRange().GetLowerBoundValue(), Time);
 	}
 
-	TimeSliderController->SetViewRange(TimeSliderController->GetViewRange().GetLowerBoundValue(), Time, EViewRangeInterpolation::Immediate);
+	TimeSliderController->SetViewRange(ViewStartTime, Time, EViewRangeInterpolation::Immediate);
 }
 
 void STimeRange::OnPlayStartTimeChanged(double NewValue)
 {
-	// We can't use the UI control to clamp the value to a Min/Max due to needing an unlimited-range spinbox for
-	// UI adjustment to work in reasonable deltas, so instead we clamp it here.
-	NewValue = FMath::Clamp(NewValue, MinPlayStartTime().GetValue(), MaxPlayStartTime().GetValue());
-
 	FFrameRate TickResolution = TimeSliderController->GetTickResolution();
 	FFrameTime Time = FFrameTime::FromDecimal(NewValue);
 	double     TimeInSeconds = TickResolution.AsSeconds(Time);
 
-	if (TimeInSeconds < TimeSliderController.Get()->GetClampRange().GetLowerBoundValue())
+	TRange<FFrameNumber> PlayRange = TimeSliderController->GetPlayRange();
+	FFrameNumber PlayDuration;
+	if (Time.FrameNumber >= MovieScene::DiscreteExclusiveUpper(PlayRange))
 	{
-		TimeSliderController->SetClampRange(TimeInSeconds, TimeSliderController->GetClampRange().GetLowerBoundValue());
+		PlayDuration = MovieScene::DiscreteExclusiveUpper(PlayRange) - MovieScene::DiscreteInclusiveLower(PlayRange);
+	}
+	else
+	{
+		PlayDuration = MovieScene::DiscreteExclusiveUpper(PlayRange) - Time.FrameNumber;
 	}
 
-	FFrameNumber PlayDuration = MovieScene::DiscreteExclusiveUpper(TimeSliderController->GetPlayRange()) - Time.FrameNumber;
 	TimeSliderController->SetPlayRange(Time.FrameNumber, PlayDuration.Value);
+
+	// Expand view ranges if outside of play range
+	if (TimeInSeconds < TimeSliderController.Get()->GetClampRange().GetLowerBoundValue())
+	{
+		OnViewStartTimeChanged(NewValue);
+	}
+
+	FFrameNumber PlayEnd = TimeSliderController->GetPlayRange().GetUpperBoundValue();
+	double PlayEndSeconds = PlayEnd / TickResolution;
+
+	if (PlayEndSeconds > TimeSliderController.Get()->GetClampRange().GetUpperBoundValue())
+	{
+		OnViewEndTimeChanged(TickResolution.AsFrameNumber(PlayEndSeconds).Value);
+	}
 }
 
 void STimeRange::OnPlayEndTimeChanged(double NewValue)
 {
-	// We can't use the UI control to clamp the value to a Min/Max due to needing an unlimited-range spinbox for
-	// UI adjustment to work in reasonable deltas, so instead we clamp it here.
-	NewValue = FMath::Clamp(NewValue, MinPlayEndTime().GetValue(), MaxPlayEndTime().GetValue());
-
 	FFrameRate TickResolution = TimeSliderController->GetTickResolution();
 	FFrameTime Time = FFrameTime::FromDecimal(NewValue);
 	double     TimeInSeconds = TickResolution.AsSeconds(Time);
 
+	TRange<FFrameNumber> PlayRange = TimeSliderController->GetPlayRange();
+	FFrameNumber PlayDuration;
+	FFrameNumber StartFrame = MovieScene::DiscreteInclusiveLower(PlayRange);
+	if (Time.FrameNumber <= StartFrame)
+	{
+		PlayDuration = MovieScene::DiscreteExclusiveUpper(PlayRange) - StartFrame;
+		StartFrame = Time.FrameNumber - PlayDuration;
+	}
+	else
+	{
+		PlayDuration = Time.FrameNumber - StartFrame;
+	}
+
+	TimeSliderController->SetPlayRange(StartFrame, PlayDuration.Value);
+
+	// Expand view ranges if outside of play range
 	if (TimeInSeconds > TimeSliderController->GetClampRange().GetUpperBoundValue())
 	{
-		TimeSliderController->SetClampRange(TimeSliderController->GetClampRange().GetLowerBoundValue(), TimeInSeconds);
+		OnViewEndTimeChanged(NewValue);
 	}
 
-	FFrameNumber StartFrame   = MovieScene::DiscreteInclusiveLower(TimeSliderController->GetPlayRange());
-	int32        PlayDuration = (Time.FrameNumber - StartFrame).Value;
+	FFrameNumber PlayStart = TimeSliderController->GetPlayRange().GetLowerBoundValue();
+	double PlayStartSeconds = PlayStart / TickResolution;
 
-	TimeSliderController->SetPlayRange(StartFrame, PlayDuration);
+	if (PlayStartSeconds < TimeSliderController.Get()->GetClampRange().GetLowerBoundValue())
+	{
+		OnViewStartTimeChanged(TickResolution.AsFrameNumber(PlayStartSeconds).Value);
+	}
 }
 
 #undef LOCTEXT_NAMESPACE
diff --git a/Engine/Source/Editor/SequencerWidgets/Private/STimeRange.h b/Engine/Source/Editor/SequencerWidgets/Private/STimeRange.h
--- a/Engine/Source/Editor/SequencerWidgets/Private/STimeRange.h
+++ b/Engine/Source/Editor/SequencerWidgets/Private/STimeRange.h
@@ -46,11 +46,6 @@ protected:
 	double PlayStartTime() const;
 	double PlayEndTime() const;
 
-	TOptional<double> MinPlayStartTime() const;
-	TOptional<double> MaxPlayStartTime() const;
-	TOptional<double> MinPlayEndTime() const;
-	TOptional<double> MaxPlayEndTime() const;
-
 	void OnPlayStartTimeCommitted(double NewValue, ETextCommit::Type InTextCommit);
 	void OnPlayEndTimeCommitted(double NewValue, ETextCommit::Type InTextCommit);
 
@@ -62,9 +57,6 @@ protected:
 	double ViewStartTime() const;
 	double ViewEndTime() const;
 	
-	TOptional<double> MaxViewStartTime() const;
-	TOptional<double> MinViewEndTime() const;
-
 	void OnViewStartTimeCommitted(double NewValue, ETextCommit::Type InTextCommit);
 	void OnViewEndTimeCommitted(double NewValue, ETextCommit::Type InTextCommit);
 
@@ -76,9 +68,6 @@ protected:
 	double WorkingStartTime() const;
 	double WorkingEndTime() const;
 
-	TOptional<double> MaxWorkingStartTime() const;
-	TOptional<double> MinWorkingEndTime() const;
-
 	void OnWorkingStartTimeCommitted(double NewValue, ETextCommit::Type InTextCommit);
 	void OnWorkingEndTimeCommitted(double NewValue, ETextCommit::Type InTextCommit);
 
