Sequencer Widgets: Fix truncation warnings

#rnx

diff --git a/Engine/Source/Editor/SequencerWidgets/Private/STimeRange.cpp b/Engine/Source/Editor/SequencerWidgets/Private/STimeRange.cpp
--- a/Engine/Source/Editor/SequencerWidgets/Private/STimeRange.cpp
+++ b/Engine/Source/Editor/SequencerWidgets/Private/STimeRange.cpp
@@ -148,7 +148,7 @@ void STimeRange::Construct( const STimeRange::FArguments& InArgs, TSharedRef<ITi
 		[
 			SNew(SBox)
 			.Visibility(InArgs._ShowWorkingRange ? EVisibility::Visible : EVisibility::Collapsed)
-			.MinDesiredWidth(64)
+			.MinDesiredWidth(64.f)
 			.HAlign(HAlign_Center)
 			[
 				WorkingRangeStart
@@ -161,12 +161,12 @@ void STimeRange::Construct( const STimeRange::FArguments& InArgs, TSharedRef<ITi
 		.Padding(2.f)
 		[
 			SNew(SBox)
-			.MinDesiredWidth(64)
+			.MinDesiredWidth(64.f)
 			.HAlign(HAlign_Center)
 			.Visibility(InArgs._ShowPlaybackRange ? EVisibility::Visible : EVisibility::Collapsed)
 			[
 				SNew(SBorder)
-				.Padding(0)
+				.Padding(0.f)
 				.BorderImage(nullptr)
 				.ForegroundColor(FLinearColor::Green)
 				[
@@ -182,7 +182,7 @@ void STimeRange::Construct( const STimeRange::FArguments& InArgs, TSharedRef<ITi
 		[
 			SNew(SBox)
 			.Visibility(InArgs._ShowViewRange ? EVisibility::Visible : EVisibility::Collapsed)
-			.MinDesiredWidth(64)
+			.MinDesiredWidth(64.f)
 			.HAlign(HAlign_Center)
 			[
 				ViewRangeStart
@@ -204,7 +204,7 @@ void STimeRange::Construct( const STimeRange::FArguments& InArgs, TSharedRef<ITi
 		[
 			SNew(SBox)
 			.Visibility(InArgs._ShowViewRange ? EVisibility::Visible : EVisibility::Collapsed)
-			.MinDesiredWidth(64)
+			.MinDesiredWidth(64.f)
 			.HAlign(HAlign_Center)
 			[
 				ViewRangeEnd
@@ -217,12 +217,12 @@ void STimeRange::Construct( const STimeRange::FArguments& InArgs, TSharedRef<ITi
 		.Padding(2.f)
 		[
 			SNew(SBox)
-			.MinDesiredWidth(64)
+			.MinDesiredWidth(64.f)
 			.HAlign(HAlign_Center)
 			.Visibility(InArgs._ShowPlaybackRange ? EVisibility::Visible : EVisibility::Collapsed)
 			[
 				SNew(SBorder)
-				.Padding(0)
+				.Padding(0.f)
 				.BorderImage(nullptr)
 				.ForegroundColor(FLinearColor::Red)
 				[
@@ -237,7 +237,7 @@ void STimeRange::Construct( const STimeRange::FArguments& InArgs, TSharedRef<ITi
 		.Padding(2.f)
 		[
 			SNew(SBox)
-			.MinDesiredWidth(64)
+			.MinDesiredWidth(64.f)
 			.HAlign(HAlign_Center)
 			.Visibility(InArgs._ShowWorkingRange ? EVisibility::Visible : EVisibility::Collapsed)
 			[
diff --git a/Engine/Source/Editor/SequencerWidgets/Private/STimeRangeSlider.cpp b/Engine/Source/Editor/SequencerWidgets/Private/STimeRangeSlider.cpp
--- a/Engine/Source/Editor/SequencerWidgets/Private/STimeRangeSlider.cpp
+++ b/Engine/Source/Editor/SequencerWidgets/Private/STimeRangeSlider.cpp
@@ -34,28 +34,28 @@ void STimeRangeSlider::Construct( const FArguments& InArgs, TSharedRef<ITimeSlid
 	ResetHoveredState();
 }
 
-float STimeRangeSlider::ComputeDragDelta(const FPointerEvent& MouseEvent, int32 GeometryWidth) const
+double STimeRangeSlider::ComputeDragDelta(const FPointerEvent& MouseEvent, double GeometryWidth) const
 {
-	float StartTime = 0;
-	float EndTime = 0;
+	double StartTime = 0;
+	double EndTime = 0;
 
 	if (TimeSliderController.IsValid())
 	{
 		StartTime = TimeSliderController->GetClampRange().GetLowerBoundValue();
 		EndTime = TimeSliderController->GetClampRange().GetUpperBoundValue();
 	}
-	float DragDistance = (MouseEvent.GetScreenSpacePosition() - MouseDownPosition).X;
+	double DragDistance = (MouseEvent.GetScreenSpacePosition() - MouseDownPosition).X;
 
-	const float PixelToUnits = (EndTime - StartTime) / (GeometryWidth - TimeRangeSliderConstants::HandleSize*2);
+	const double PixelToUnits = (EndTime - StartTime) / (GeometryWidth - TimeRangeSliderConstants::HandleSize*2);
 	return DragDistance * PixelToUnits;
 }
 
-void STimeRangeSlider::ComputeHandleOffsets(float& LeftHandleOffset, float& HandleOffset, float& RightHandleOffset, int32 GeometryWidth) const
+void STimeRangeSlider::ComputeHandleOffsets(double& LeftHandleOffset, double& HandleOffset, double& RightHandleOffset, double GeometryWidth) const
 {
-	float StartTime = 0;
-	float InTime = 0;
-	float OutTime = 0;
-	float EndTime = 0;
+	double StartTime = 0;
+	double InTime = 0;
+	double OutTime = 0;
+	double EndTime = 0;
 
 	if (TimeSliderController.IsValid())
 	{
@@ -65,18 +65,18 @@ void STimeRangeSlider::ComputeHandleOffsets(float& LeftHandleOffset, float& Hand
 		EndTime = TimeSliderController->GetClampRange().GetUpperBoundValue();
 	}
 
-	const float UnitsToPixel = float(GeometryWidth - TimeRangeSliderConstants::HandleSize*2) / (EndTime - StartTime);
+	const double UnitsToPixel = (GeometryWidth - TimeRangeSliderConstants::HandleSize*2) / (EndTime - StartTime);
 
 	LeftHandleOffset = (InTime - StartTime) * UnitsToPixel;
 	HandleOffset = LeftHandleOffset + TimeRangeSliderConstants::HandleSize;
 	RightHandleOffset = HandleOffset + (OutTime - InTime) * UnitsToPixel;
 	
-	float ScrubberWidth = RightHandleOffset-LeftHandleOffset-TimeRangeSliderConstants::HandleSize;
-	if (ScrubberWidth < (float)TimeRangeSliderConstants::MinimumScrubberWidth)
+	double ScrubberWidth = RightHandleOffset-LeftHandleOffset-TimeRangeSliderConstants::HandleSize;
+	if (ScrubberWidth < (double)TimeRangeSliderConstants::MinimumScrubberWidth)
 	{
-		HandleOffset = HandleOffset - ((float)TimeRangeSliderConstants::MinimumScrubberWidth - ScrubberWidth) / 2.f;
+		HandleOffset = HandleOffset - (TimeRangeSliderConstants::MinimumScrubberWidth - ScrubberWidth) / 2.0;
 		LeftHandleOffset = HandleOffset - TimeRangeSliderConstants::HandleSize;
-		RightHandleOffset = HandleOffset + (float)TimeRangeSliderConstants::MinimumScrubberWidth;
+		RightHandleOffset = HandleOffset + TimeRangeSliderConstants::MinimumScrubberWidth;
 	}
 }
 
@@ -97,9 +97,9 @@ int32 STimeRangeSlider::OnPaint( const FPaintArgs& Args, const FGeometry& Allott
 	static const FSlateBrush* RangeHandleRight = FAppStyle::GetBrush( TEXT( "Sequencer.Timeline.RangeHandleRight" ) ); 
 	static const FSlateBrush* RangeHandle = FAppStyle::GetBrush( TEXT( "Sequencer.Timeline.RangeHandle" ) ); 
 
-	float LeftHandleOffset = 0.f;
-	float HandleOffset = 0.f;
-	float RightHandleOffset = 0.f;
+	double LeftHandleOffset = 0;
+	double HandleOffset = 0;
+	double RightHandleOffset = 0;
 	ComputeHandleOffsets(LeftHandleOffset, HandleOffset, RightHandleOffset, AllottedGeometry.GetLocalSize().X);
 
 	static const FName SelectionColorName("SelectionColor");
@@ -175,7 +175,7 @@ FReply STimeRangeSlider::OnMouseMove( const FGeometry& MyGeometry, const FPointe
 {
 	if (HasMouseCapture())
 	{
-		float DragDelta = ComputeDragDelta(MouseEvent, MyGeometry.GetLocalSize().X);
+		double DragDelta = ComputeDragDelta(MouseEvent, MyGeometry.GetLocalSize().X);
 
 		ITimeSliderController* TimeSliderControllerPtr = TimeSliderController.Get();
 		if (!TimeSliderControllerPtr)
@@ -204,8 +204,8 @@ FReply STimeRangeSlider::OnMouseMove( const FGeometry& MyGeometry, const FPointe
 		}
 		else if (bLeftHandleDragged || bRightHandleDragged)
 		{
-			double NewIn = 0.0f;
-			double NewOut = 0.0f;
+			double NewIn = 0;
+			double NewOut = 0;
 
 			if (bLeftHandleDragged)
 			{
@@ -244,14 +244,14 @@ FReply STimeRangeSlider::OnMouseMove( const FGeometry& MyGeometry, const FPointe
 	{
 		ResetHoveredState();
 
-		float LeftHandleOffset = 0.f;
-		float HandleOffset = 0.f;
-		float RightHandleOffset = 0.f;
+		double LeftHandleOffset = 0;
+		double HandleOffset = 0;
+		double RightHandleOffset = 0;
 		ComputeHandleOffsets(LeftHandleOffset, HandleOffset, RightHandleOffset, MyGeometry.GetLocalSize().X);
 		
-		FGeometry LeftHandleRect  = MyGeometry.MakeChild(FVector2D(LeftHandleOffset, 0.f), FVector2D(TimeRangeSliderConstants::HandleSize, TimeRangeSliderConstants::HandleSize));
-		FGeometry RightHandleRect = MyGeometry.MakeChild(FVector2D(RightHandleOffset, 0.f), FVector2D(TimeRangeSliderConstants::HandleSize, TimeRangeSliderConstants::HandleSize));
-		FGeometry HandleRect      = MyGeometry.MakeChild(FVector2D(HandleOffset, 0.0f), FVector2D(RightHandleOffset-LeftHandleOffset-TimeRangeSliderConstants::HandleSize, TimeRangeSliderConstants::HandleSize));
+		FGeometry LeftHandleRect  = MyGeometry.MakeChild(FVector2D(LeftHandleOffset, 0), FVector2D(TimeRangeSliderConstants::HandleSize, TimeRangeSliderConstants::HandleSize));
+		FGeometry RightHandleRect = MyGeometry.MakeChild(FVector2D(RightHandleOffset, 0), FVector2D(TimeRangeSliderConstants::HandleSize, TimeRangeSliderConstants::HandleSize));
+		FGeometry HandleRect      = MyGeometry.MakeChild(FVector2D(HandleOffset, 0), FVector2D(RightHandleOffset-LeftHandleOffset-TimeRangeSliderConstants::HandleSize, TimeRangeSliderConstants::HandleSize));
 
 		FVector2D LocalMousePosition = MouseEvent.GetScreenSpacePosition();
 
diff --git a/Engine/Source/Editor/SequencerWidgets/Private/STimeRangeSlider.h b/Engine/Source/Editor/SequencerWidgets/Private/STimeRangeSlider.h
--- a/Engine/Source/Editor/SequencerWidgets/Private/STimeRangeSlider.h
+++ b/Engine/Source/Editor/SequencerWidgets/Private/STimeRangeSlider.h
@@ -40,8 +40,8 @@ public:
 protected:
 	void ResetState();
 	void ResetHoveredState();
-	float ComputeDragDelta(const FPointerEvent& MouseEvent, int32 GeometryWidth) const;
-	void ComputeHandleOffsets(float& LeftHandleOffset, float& RightHandleOffset, float&HandleOffset, int32 GeometryWidth) const;
+	double ComputeDragDelta(const FPointerEvent& MouseEvent, double GeometryWidth) const;
+	void ComputeHandleOffsets(double& LeftHandleOffset, double& RightHandleOffset, double&HandleOffset, double GeometryWidth) const;
 
 private:
 	/* The left handle is being dragged */
diff --git a/Engine/Source/Editor/SequencerWidgets/SequencerWidgets.Build.cs b/Engine/Source/Editor/SequencerWidgets/SequencerWidgets.Build.cs
--- a/Engine/Source/Editor/SequencerWidgets/SequencerWidgets.Build.cs
+++ b/Engine/Source/Editor/SequencerWidgets/SequencerWidgets.Build.cs
@@ -6,6 +6,8 @@ public class SequencerWidgets : ModuleRules
 {
 	public SequencerWidgets(ReadOnlyTargetRules Target) : base(Target)
 	{
+		UnsafeTypeCastWarningLevel = WarningLevel.Warning;
+
 		PublicIncludePathModuleNames.AddRange(
 			new string[] {
 				"SequencerCore",
