Sequencer Widgets: Remove warning for now
#rnx

diff --git a/Engine/Source/Editor/SequencerWidgets/SequencerWidgets.Build.cs b/Engine/Source/Editor/SequencerWidgets/SequencerWidgets.Build.cs
--- a/Engine/Source/Editor/SequencerWidgets/SequencerWidgets.Build.cs
+++ b/Engine/Source/Editor/SequencerWidgets/SequencerWidgets.Build.cs
@@ -6,8 +6,6 @@ public class SequencerWidgets : ModuleRules
 {
 	public SequencerWidgets(ReadOnlyTargetRules Target) : base(Target)
 	{
-		UnsafeTypeCastWarningLevel = WarningLevel.Warning;
-
 		PublicIncludePathModuleNames.AddRange(
 			new string[] {
 				"SequencerCore",
