Cleaning up ITimeSlider dependencies.

diff --git a/Engine/Plugins/Animation/ControlRig/Source/ControlRigEditor/Private/EditMode/SControlRigDetails.cpp b/Engine/Plugins/Animation/ControlRig/Source/ControlRigEditor/Private/EditMode/SControlRigDetails.cpp
--- a/Engine/Plugins/Animation/ControlRig/Source/ControlRigEditor/Private/EditMode/SControlRigDetails.cpp
+++ b/Engine/Plugins/Animation/ControlRig/Source/ControlRigEditor/Private/EditMode/SControlRigDetails.cpp
@@ -14,6 +14,7 @@
 #include "EditorModeManager.h"
 #include "ISequencer.h"
 #include "LevelSequence.h"
+#include "MovieScene.h"
 #include "Selection.h"
 #include "Editor.h"
 #include "LevelEditor.h"
diff --git a/Engine/Plugins/Animation/ControlRig/Source/ControlRigEditor/Private/EditMode/SControlRigSnapper.cpp b/Engine/Plugins/Animation/ControlRig/Source/ControlRigEditor/Private/EditMode/SControlRigSnapper.cpp
--- a/Engine/Plugins/Animation/ControlRig/Source/ControlRigEditor/Private/EditMode/SControlRigSnapper.cpp
+++ b/Engine/Plugins/Animation/ControlRig/Source/ControlRigEditor/Private/EditMode/SControlRigSnapper.cpp
@@ -29,6 +29,7 @@
 #include "Widgets/Input/SButton.h"
 #include "Widgets/Layout/SScrollBox.h"
 #include "Tools/ControlRigSnapSettings.h"
+#include "MovieScene.h"
 
 #define LOCTEXT_NAMESPACE "ControlRigSnapper"
 
diff --git a/Engine/Plugins/Animation/MLDeformer/MLDeformerFramework/Source/MLDeformerFrameworkEditor/Private/SMLDeformerTimeline.cpp b/Engine/Plugins/Animation/MLDeformer/MLDeformerFramework/Source/MLDeformerFrameworkEditor/Private/SMLDeformerTimeline.cpp
--- a/Engine/Plugins/Animation/MLDeformer/MLDeformerFramework/Source/MLDeformerFrameworkEditor/Private/SMLDeformerTimeline.cpp
+++ b/Engine/Plugins/Animation/MLDeformer/MLDeformerFramework/Source/MLDeformerFrameworkEditor/Private/SMLDeformerTimeline.cpp
@@ -16,11 +16,13 @@
 #include "MLDeformerEditorModel.h"
 #include "MLDeformerModel.h"
 #include "Modules/ModuleManager.h"
+#include "MovieSceneFwd.h"
 #include "MovieSceneTimeHelpers.h"
 #include "Preferences/PersonaOptions.h"
 #include "ScopedTransaction.h"
 #include "Styling/AppStyle.h"
 #include "Styling/ISlateStyle.h"
+#include "TimeSliderArgs.h"
 #include "Widgets/Input/SButton.h"
 #include "Widgets/Input/SSearchBox.h"
 #include "Widgets/Input/SSpinBox.h"
diff --git a/Engine/Plugins/Animation/MLDeformer/MLDeformerFramework/Source/MLDeformerFrameworkEditor/Private/SMLDeformerTimeline.h b/Engine/Plugins/Animation/MLDeformer/MLDeformerFramework/Source/MLDeformerFrameworkEditor/Private/SMLDeformerTimeline.h
--- a/Engine/Plugins/Animation/MLDeformer/MLDeformerFramework/Source/MLDeformerFrameworkEditor/Private/SMLDeformerTimeline.h
+++ b/Engine/Plugins/Animation/MLDeformer/MLDeformerFramework/Source/MLDeformerFrameworkEditor/Private/SMLDeformerTimeline.h
@@ -2,6 +2,7 @@
 
 #pragma once
 
+#include "AnimatedRange.h"
 #include "CoreTypes.h"
 #include "Templates/SharedPointer.h"
 #include "Widgets/DeclarativeSyntaxSupport.h"
diff --git a/Engine/Plugins/Developer/Concert/ConcertSync/ConcertSyncClient/Source/ConcertSyncClient/Private/ConcertClientSequencerManager.cpp b/Engine/Plugins/Developer/Concert/ConcertSync/ConcertSyncClient/Source/ConcertSyncClient/Private/ConcertClientSequencerManager.cpp
--- a/Engine/Plugins/Developer/Concert/ConcertSync/ConcertSyncClient/Source/ConcertSyncClient/Private/ConcertClientSequencerManager.cpp
+++ b/Engine/Plugins/Developer/Concert/ConcertSync/ConcertSyncClient/Source/ConcertSyncClient/Private/ConcertClientSequencerManager.cpp
@@ -27,6 +27,7 @@
 #include "UObject/UObjectGlobals.h"
 
 #if WITH_EDITOR
+	#include "AnimatedRange.h"
 	#include "ISequencerModule.h"
 	#include "ISequencer.h"
 	#include "SequencerSettings.h"
diff --git a/Engine/Plugins/Editor/CurveEditorTools/Source/CurveEditorTools/Private/CurveEditorFocusExtension.cpp b/Engine/Plugins/Editor/CurveEditorTools/Source/CurveEditorTools/Private/CurveEditorFocusExtension.cpp
--- a/Engine/Plugins/Editor/CurveEditorTools/Source/CurveEditorTools/Private/CurveEditorFocusExtension.cpp
+++ b/Engine/Plugins/Editor/CurveEditorTools/Source/CurveEditorTools/Private/CurveEditorFocusExtension.cpp
@@ -5,6 +5,7 @@
 #include "CurveEditorToolCommands.h"
 #include "CurveEditor.h"
 #include "ICurveEditorBounds.h"
+#include "AnimatedRange.h"
 
 void FCurveEditorFocusExtension::BindCommands(TSharedRef<FUICommandList> CommandBindings)
 {
diff --git a/Engine/Plugins/FX/Niagara/Source/NiagaraEditor/Private/Widgets/Scalability/SNiagaraScalabilityContext.h b/Engine/Plugins/FX/Niagara/Source/NiagaraEditor/Private/Widgets/Scalability/SNiagaraScalabilityContext.h
--- a/Engine/Plugins/FX/Niagara/Source/NiagaraEditor/Private/Widgets/Scalability/SNiagaraScalabilityContext.h
+++ b/Engine/Plugins/FX/Niagara/Source/NiagaraEditor/Private/Widgets/Scalability/SNiagaraScalabilityContext.h
@@ -7,6 +7,7 @@
 #include "NiagaraSystem.h"
 #include "ViewModels/NiagaraSystemScalabilityViewModel.h"
 #include "Widgets/Layout/SGridPanel.h"
+#include "Widgets/SCompoundWidget.h"
 
 struct FPropertyAndParent;
 
diff --git a/Engine/Plugins/FX/Niagara/Source/NiagaraEditorWidgets/Private/SNiagaraOverviewGraphTitleBar.h b/Engine/Plugins/FX/Niagara/Source/NiagaraEditorWidgets/Private/SNiagaraOverviewGraphTitleBar.h
--- a/Engine/Plugins/FX/Niagara/Source/NiagaraEditorWidgets/Private/SNiagaraOverviewGraphTitleBar.h
+++ b/Engine/Plugins/FX/Niagara/Source/NiagaraEditorWidgets/Private/SNiagaraOverviewGraphTitleBar.h
@@ -7,6 +7,7 @@
 #include "NiagaraSystem.h"
 #include "ViewModels/NiagaraSystemScalabilityViewModel.h"
 #include "ViewModels/NiagaraSystemViewModel.h"
+#include "Widgets/SCompoundWidget.h"
 //#include "SNiagaraOverviewGraph.h"
 
 class SNiagaraOverviewGraph;
diff --git a/Engine/Plugins/MovieScene/LevelSequenceEditor/Source/LevelSequenceEditor/Private/CinematicViewport/SCinematicLevelViewport.cpp b/Engine/Plugins/MovieScene/LevelSequenceEditor/Source/LevelSequenceEditor/Private/CinematicViewport/SCinematicLevelViewport.cpp
--- a/Engine/Plugins/MovieScene/LevelSequenceEditor/Source/LevelSequenceEditor/Private/CinematicViewport/SCinematicLevelViewport.cpp
+++ b/Engine/Plugins/MovieScene/LevelSequenceEditor/Source/LevelSequenceEditor/Private/CinematicViewport/SCinematicLevelViewport.cpp
@@ -1,11 +1,13 @@
 // Copyright Epic Games, Inc. All Rights Reserved.
 
 #include "CinematicViewport/SCinematicLevelViewport.h"
+#include "AnimatedRange.h"
 #include "EditorModeManager.h"
 #include "Framework/Commands/UICommandList.h"
 #include "Evaluation/MovieSceneSequenceTransform.h"
 #include "Widgets/Layout/SBorder.h"
 #include "LevelSequenceEditorToolkit.h"
+#include "MovieScene.h"
 #include "Rendering/SlateRenderer.h"
 #include "SlateOptMacros.h"
 #include "Widgets/Layout/SSpacer.h"
diff --git a/Engine/Plugins/MovieScene/LevelSequenceEditor/Source/LevelSequenceEditor/Private/CinematicViewport/SCinematicTransportRange.cpp b/Engine/Plugins/MovieScene/LevelSequenceEditor/Source/LevelSequenceEditor/Private/CinematicViewport/SCinematicTransportRange.cpp
--- a/Engine/Plugins/MovieScene/LevelSequenceEditor/Source/LevelSequenceEditor/Private/CinematicViewport/SCinematicTransportRange.cpp
+++ b/Engine/Plugins/MovieScene/LevelSequenceEditor/Source/LevelSequenceEditor/Private/CinematicViewport/SCinematicTransportRange.cpp
@@ -5,6 +5,7 @@
 #include "SequencerKeyCollection.h"
 #include "SequencerSettings.h"
 #include "Misc/QualifiedFrameTime.h"
+#include "MovieScene.h"
 #include "MovieSceneTimeHelpers.h"
 #include "Styling/AppStyle.h"
 #include "Styles/LevelSequenceEditorStyle.h"
diff --git a/Engine/Plugins/MovieScene/LevelSequenceEditor/Source/LevelSequenceEditor/Private/LevelSequenceEditorSubsystem.cpp b/Engine/Plugins/MovieScene/LevelSequenceEditor/Source/LevelSequenceEditor/Private/LevelSequenceEditorSubsystem.cpp
--- a/Engine/Plugins/MovieScene/LevelSequenceEditor/Source/LevelSequenceEditor/Private/LevelSequenceEditorSubsystem.cpp
+++ b/Engine/Plugins/MovieScene/LevelSequenceEditor/Source/LevelSequenceEditor/Private/LevelSequenceEditorSubsystem.cpp
@@ -10,6 +10,7 @@
 #include "LevelSequenceEditorCommands.h"
 #include "MovieScenePossessable.h"
 #include "SequencerSettings.h"
+#include "MovieScene.h"
 #include "MovieSceneSpawnable.h"
 #include "SequencerUtilities.h"
 #include "Sections/MovieScene3DConstraintSection.h"
diff --git a/Engine/Plugins/MovieScene/LevelSequenceEditor/Source/LevelSequenceEditor/Private/Misc/LevelSequenceEditorActorBinding.cpp b/Engine/Plugins/MovieScene/LevelSequenceEditor/Source/LevelSequenceEditor/Private/Misc/LevelSequenceEditorActorBinding.cpp
--- a/Engine/Plugins/MovieScene/LevelSequenceEditor/Source/LevelSequenceEditor/Private/Misc/LevelSequenceEditorActorBinding.cpp
+++ b/Engine/Plugins/MovieScene/LevelSequenceEditor/Source/LevelSequenceEditor/Private/Misc/LevelSequenceEditorActorBinding.cpp
@@ -4,6 +4,7 @@
 #include "ISceneOutliner.h"
 #include "ISequencer.h"
 #include "LevelSequence.h"
+#include "MovieScene.h"
 #include "MovieScenePossessable.h"
 #include "Styling/SlateIconFinder.h"
 #include "SceneOutlinerModule.h"
diff --git a/Engine/Plugins/MovieScene/LevelSequenceEditor/Source/LevelSequenceEditor/Private/Misc/LevelSequenceEditorActorSpawner.cpp b/Engine/Plugins/MovieScene/LevelSequenceEditor/Source/LevelSequenceEditor/Private/Misc/LevelSequenceEditorActorSpawner.cpp
--- a/Engine/Plugins/MovieScene/LevelSequenceEditor/Source/LevelSequenceEditor/Private/Misc/LevelSequenceEditorActorSpawner.cpp
+++ b/Engine/Plugins/MovieScene/LevelSequenceEditor/Source/LevelSequenceEditor/Private/Misc/LevelSequenceEditorActorSpawner.cpp
@@ -13,6 +13,7 @@
 #include "SequencerSettings.h"
 #include "IMovieScenePlaybackClient.h"
 #include "ISequencer.h"
+#include "MovieScene.h"
 #include "Tracks/MovieSceneSpawnTrack.h"
 #include "Sections/MovieSceneBoolSection.h"
 #include "Settings/LevelEditorViewportSettings.h"
diff --git a/Engine/Plugins/MovieScene/LevelSequenceEditor/Source/LevelSequenceEditor/Private/Misc/LevelSequenceEditorSpawnRegister.cpp b/Engine/Plugins/MovieScene/LevelSequenceEditor/Source/LevelSequenceEditor/Private/Misc/LevelSequenceEditorSpawnRegister.cpp
--- a/Engine/Plugins/MovieScene/LevelSequenceEditor/Source/LevelSequenceEditor/Private/Misc/LevelSequenceEditorSpawnRegister.cpp
+++ b/Engine/Plugins/MovieScene/LevelSequenceEditor/Source/LevelSequenceEditor/Private/Misc/LevelSequenceEditorSpawnRegister.cpp
@@ -6,6 +6,7 @@
 #include "LevelEditor.h"
 #include "EntitySystem/MovieSceneSpawnablesSystem.h"
 #include "Evaluation/MovieSceneEvaluationTemplateInstance.h"
+#include "MovieScene.h"
 #include "MovieSceneSpawnable.h"
 #include "UObject/ObjectSaveContext.h"
 #include "Selection.h"
diff --git a/Engine/Plugins/MovieScene/TemplateSequence/Source/TemplateSequenceEditor/Private/Misc/TemplateSequenceEditorUtil.cpp b/Engine/Plugins/MovieScene/TemplateSequence/Source/TemplateSequenceEditor/Private/Misc/TemplateSequenceEditorUtil.cpp
--- a/Engine/Plugins/MovieScene/TemplateSequence/Source/TemplateSequenceEditor/Private/Misc/TemplateSequenceEditorUtil.cpp
+++ b/Engine/Plugins/MovieScene/TemplateSequence/Source/TemplateSequenceEditor/Private/Misc/TemplateSequenceEditorUtil.cpp
@@ -2,6 +2,7 @@
 
 #include "Misc/TemplateSequenceEditorUtil.h"
 #include "ISequencer.h"
+#include "MovieScene.h"
 #include "MovieSceneSpawnable.h"
 #include "TemplateSequence.h"
 
diff --git a/Engine/Plugins/MovieScene/TemplateSequence/Source/TemplateSequenceEditor/Private/TemplateSequenceEditorToolkit.cpp b/Engine/Plugins/MovieScene/TemplateSequence/Source/TemplateSequenceEditor/Private/TemplateSequenceEditorToolkit.cpp
--- a/Engine/Plugins/MovieScene/TemplateSequence/Source/TemplateSequenceEditor/Private/TemplateSequenceEditorToolkit.cpp
+++ b/Engine/Plugins/MovieScene/TemplateSequence/Source/TemplateSequenceEditor/Private/TemplateSequenceEditorToolkit.cpp
@@ -8,6 +8,7 @@
 #include "Misc/TemplateSequenceEditorPlaybackContext.h"
 #include "Misc/TemplateSequenceEditorSpawnRegister.h"
 #include "Modules/ModuleManager.h"
+#include "MovieScene.h"
 #include "ScopedTransaction.h"
 #include "Selection.h"
 #include "Widgets/Docking/SDockTab.h"
diff --git a/Engine/Plugins/VirtualProduction/CameraCalibration/Source/CameraCalibrationEditor/CameraCalibrationEditor.Build.cs b/Engine/Plugins/VirtualProduction/CameraCalibration/Source/CameraCalibrationEditor/CameraCalibrationEditor.Build.cs
--- a/Engine/Plugins/VirtualProduction/CameraCalibration/Source/CameraCalibrationEditor/CameraCalibrationEditor.Build.cs
+++ b/Engine/Plugins/VirtualProduction/CameraCalibration/Source/CameraCalibrationEditor/CameraCalibrationEditor.Build.cs
@@ -35,6 +35,7 @@ namespace UnrealBuildTool.Rules
 					"OpenCVHelper",
 					"PlacementMode",
 					"PropertyEditor",
+					"SequencerWidgets",
 					"Settings",
 					"Slate",
 					"SlateCore",
diff --git a/Engine/Plugins/VirtualProduction/CameraCalibration/Source/CameraCalibrationEditor/Private/AssetEditor/CameraCalibrationTimeSliderController.h b/Engine/Plugins/VirtualProduction/CameraCalibration/Source/CameraCalibrationEditor/Private/AssetEditor/CameraCalibrationTimeSliderController.h
--- a/Engine/Plugins/VirtualProduction/CameraCalibration/Source/CameraCalibrationEditor/Private/AssetEditor/CameraCalibrationTimeSliderController.h
+++ b/Engine/Plugins/VirtualProduction/CameraCalibration/Source/CameraCalibrationEditor/Private/AssetEditor/CameraCalibrationTimeSliderController.h
@@ -4,6 +4,7 @@
 
 #include "CoreMinimal.h"
 #include "ITimeSlider.h"
+#include "TimeSliderArgs.h"
 #include "TickableEditorObject.h"
 
 class FCameraCalibrationStepsController;
diff --git a/Engine/Plugins/VirtualProduction/SequencerPlaylists/Source/SequencerPlaylists/Private/SequencerPlaylistItem_Sequence.cpp b/Engine/Plugins/VirtualProduction/SequencerPlaylists/Source/SequencerPlaylists/Private/SequencerPlaylistItem_Sequence.cpp
--- a/Engine/Plugins/VirtualProduction/SequencerPlaylists/Source/SequencerPlaylists/Private/SequencerPlaylistItem_Sequence.cpp
+++ b/Engine/Plugins/VirtualProduction/SequencerPlaylists/Source/SequencerPlaylists/Private/SequencerPlaylistItem_Sequence.cpp
@@ -4,6 +4,7 @@
 #include "Evaluation/MovieSceneSequenceTransform.h"
 #include "Framework/Notifications/NotificationManager.h"
 #include "LevelSequence.h"
+#include "MovieScene.h"
 #include "MovieSceneFolder.h"
 #include "Recorder/TakeRecorder.h"
 #include "SequencerPlaylist.h"
diff --git a/Engine/Plugins/VirtualProduction/SequencerPlaylists/Source/SequencerPlaylists/Private/SequencerPlaylistPlayer.cpp b/Engine/Plugins/VirtualProduction/SequencerPlaylists/Source/SequencerPlaylists/Private/SequencerPlaylistPlayer.cpp
--- a/Engine/Plugins/VirtualProduction/SequencerPlaylists/Source/SequencerPlaylists/Private/SequencerPlaylistPlayer.cpp
+++ b/Engine/Plugins/VirtualProduction/SequencerPlaylists/Source/SequencerPlaylists/Private/SequencerPlaylistPlayer.cpp
@@ -6,9 +6,11 @@
 #include "SequencerPlaylistsLog.h"
 #include "SequencerPlaylistsModule.h"
 
+#include "AnimatedRange.h"
 #include "ILevelSequenceEditorToolkit.h"
 #include "LevelSequence.h"
 #include "LevelSequenceEditorBlueprintLibrary.h"
+#include "MovieScene.h"
 #include "Recorder/TakeRecorder.h"
 #include "ScopedTransaction.h"
 #include "TakePreset.h"
diff --git a/Engine/Plugins/VirtualProduction/Takes/Source/TakeRecorder/Private/Recorder/TakeRecorder.cpp b/Engine/Plugins/VirtualProduction/Takes/Source/TakeRecorder/Private/Recorder/TakeRecorder.cpp
--- a/Engine/Plugins/VirtualProduction/Takes/Source/TakeRecorder/Private/Recorder/TakeRecorder.cpp
+++ b/Engine/Plugins/VirtualProduction/Takes/Source/TakeRecorder/Private/Recorder/TakeRecorder.cpp
@@ -3,6 +3,7 @@
 #include "Recorder/TakeRecorder.h"
 
 #include "Algo/Accumulate.h"
+#include "AnimatedRange.h"
 #include "AssetRegistry/AssetRegistryModule.h"
 #include "CoreGlobals.h"
 #include "Engine/Engine.h"
diff --git a/Engine/Plugins/VirtualProduction/Takes/Source/TakeRecorder/Private/Widgets/STakeRecorderCockpit.cpp b/Engine/Plugins/VirtualProduction/Takes/Source/TakeRecorder/Private/Widgets/STakeRecorderCockpit.cpp
--- a/Engine/Plugins/VirtualProduction/Takes/Source/TakeRecorder/Private/Widgets/STakeRecorderCockpit.cpp
+++ b/Engine/Plugins/VirtualProduction/Takes/Source/TakeRecorder/Private/Widgets/STakeRecorderCockpit.cpp
@@ -14,12 +14,14 @@
 #include "TakeRecorderSources.h"
 #include "Misc/Timecode.h"
 #include "Misc/FrameRate.h"
+#include "MovieScene.h"
 #include "Recorder/TakeRecorder.h"
 #include "Recorder/TakeRecorderBlueprintLibrary.h"
 #include "LevelSequence.h"
 #include "Algo/Find.h"
 #include "Framework/MultiBox/MultiBoxBuilder.h"
 
+#include "MovieScene.h"
 #include "MovieSceneToolsProjectSettings.h"
 
 // AssetRegistry includes
diff --git a/Engine/Plugins/VirtualProduction/Takes/Source/TakeRecorder/Private/Widgets/STakeRecorderPanel.cpp b/Engine/Plugins/VirtualProduction/Takes/Source/TakeRecorder/Private/Widgets/STakeRecorderPanel.cpp
--- a/Engine/Plugins/VirtualProduction/Takes/Source/TakeRecorder/Private/Widgets/STakeRecorderPanel.cpp
+++ b/Engine/Plugins/VirtualProduction/Takes/Source/TakeRecorder/Private/Widgets/STakeRecorderPanel.cpp
@@ -19,6 +19,7 @@
 
 // Core includes
 #include "Modules/ModuleManager.h"
+#include "MovieScene.h"
 #include "Algo/Sort.h"
 #include "Misc/FileHelper.h"
 
diff --git a/Engine/Source/Editor/CurveEditor/Private/Views/SInteractiveCurveEditorView.cpp b/Engine/Source/Editor/CurveEditor/Private/Views/SInteractiveCurveEditorView.cpp
--- a/Engine/Source/Editor/CurveEditor/Private/Views/SInteractiveCurveEditorView.cpp
+++ b/Engine/Source/Editor/CurveEditor/Private/Views/SInteractiveCurveEditorView.cpp
@@ -3,6 +3,7 @@
 #include "Views/SInteractiveCurveEditorView.h"
 
 #include "Algo/Sort.h"
+#include "AnimatedRange.h"
 #include "Containers/ArrayView.h"
 #include "Containers/Map.h"
 #include "Containers/SortedMap.h"
diff --git a/Engine/Source/Editor/MovieSceneTools/Private/MovieSceneEventCustomization.cpp b/Engine/Source/Editor/MovieSceneTools/Private/MovieSceneEventCustomization.cpp
--- a/Engine/Source/Editor/MovieSceneTools/Private/MovieSceneEventCustomization.cpp
+++ b/Engine/Source/Editor/MovieSceneTools/Private/MovieSceneEventCustomization.cpp
@@ -9,6 +9,7 @@
 #include "Tracks/MovieSceneEventTrack.h"
 #include "Sections/MovieSceneEventSectionBase.h"
 #include "Algo/Find.h"
+#include "MovieScene.h"
 #include "MovieSceneSequence.h"
 #include "ISequencerModule.h"
 #include "ScopedTransaction.h"
diff --git a/Engine/Source/Editor/MovieSceneTools/Private/Sections/ThumbnailSection.cpp b/Engine/Source/Editor/MovieSceneTools/Private/Sections/ThumbnailSection.cpp
--- a/Engine/Source/Editor/MovieSceneTools/Private/Sections/ThumbnailSection.cpp
+++ b/Engine/Source/Editor/MovieSceneTools/Private/Sections/ThumbnailSection.cpp
@@ -1,6 +1,7 @@
 // Copyright Epic Games, Inc. All Rights Reserved.
 
 #include "Sections/ThumbnailSection.h"
+#include "AnimatedRange.h"
 #include "Rendering/DrawElements.h"
 #include "Textures/SlateIcon.h"
 #include "Framework/Commands/UIAction.h"
diff --git a/Engine/Source/Editor/MovieSceneTools/Private/TrackEditors/AudioTrackEditor.cpp b/Engine/Source/Editor/MovieSceneTools/Private/TrackEditors/AudioTrackEditor.cpp
--- a/Engine/Source/Editor/MovieSceneTools/Private/TrackEditors/AudioTrackEditor.cpp
+++ b/Engine/Source/Editor/MovieSceneTools/Private/TrackEditors/AudioTrackEditor.cpp
@@ -9,6 +9,7 @@
 #include "RenderUtils.h"
 #include "RenderingThread.h"
 #include "Modules/ModuleManager.h"
+#include "AnimatedRange.h"
 #include "Audio.h"
 #include "Sound/SoundBase.h"
 #include "Sound/SoundWave.h"
diff --git a/Engine/Source/Editor/MovieSceneTools/Private/TrackEditors/ParticleTrackEditor.cpp b/Engine/Source/Editor/MovieSceneTools/Private/TrackEditors/ParticleTrackEditor.cpp
--- a/Engine/Source/Editor/MovieSceneTools/Private/TrackEditors/ParticleTrackEditor.cpp
+++ b/Engine/Source/Editor/MovieSceneTools/Private/TrackEditors/ParticleTrackEditor.cpp
@@ -1,6 +1,7 @@
 // Copyright Epic Games, Inc. All Rights Reserved.
 
 #include "TrackEditors/ParticleTrackEditor.h"
+#include "AnimatedRange.h"
 #include "Rendering/DrawElements.h"
 #include "Framework/MultiBox/MultiBoxBuilder.h"
 #include "Curves/IntegralCurve.h"
diff --git a/Engine/Source/Editor/MovieSceneTools/Private/TrackEditors/SubTrackEditorBase.cpp b/Engine/Source/Editor/MovieSceneTools/Private/TrackEditors/SubTrackEditorBase.cpp
--- a/Engine/Source/Editor/MovieSceneTools/Private/TrackEditors/SubTrackEditorBase.cpp
+++ b/Engine/Source/Editor/MovieSceneTools/Private/TrackEditors/SubTrackEditorBase.cpp
@@ -4,6 +4,7 @@
 #include "Fonts/FontCache.h"
 #include "FrameNumberDisplayFormat.h"
 #include "FrameNumberNumericInterface.h"
+#include "MovieScene.h"
 #include "SequencerSettings.h"
 #include "Tracks/MovieSceneCinematicShotTrack.h"
 #include "Tracks/MovieSceneSubTrack.h"
diff --git a/Engine/Source/Editor/MovieSceneTools/Public/Channels/PerlinNoiseChannelInterface.h b/Engine/Source/Editor/MovieSceneTools/Public/Channels/PerlinNoiseChannelInterface.h
--- a/Engine/Source/Editor/MovieSceneTools/Public/Channels/PerlinNoiseChannelInterface.h
+++ b/Engine/Source/Editor/MovieSceneTools/Public/Channels/PerlinNoiseChannelInterface.h
@@ -9,6 +9,7 @@
 #include "MovieSceneSection.h"
 #include "TimeToPixel.h"
 #include "Widgets/Input/NumericTypeInterface.h"
+#include "Widgets/SCompoundWidget.h"
 #include "Widgets/Text/STextBlock.h"
 
 class FMenuBuilder;
diff --git a/Engine/Source/Editor/Persona/Private/AnimTimeline/AnimModel.h b/Engine/Source/Editor/Persona/Private/AnimTimeline/AnimModel.h
--- a/Engine/Source/Editor/Persona/Private/AnimTimeline/AnimModel.h
+++ b/Engine/Source/Editor/Persona/Private/AnimTimeline/AnimModel.h
@@ -2,6 +2,7 @@
 
 #pragma once
 
+#include "AnimatedRange.h"
 #include "Templates/SharedPointer.h"
 #include "ITimeSlider.h"
 #include "PersonaDelegates.h"
diff --git a/Engine/Source/Editor/Persona/Private/AnimTimeline/AnimTimeSliderController.h b/Engine/Source/Editor/Persona/Private/AnimTimeline/AnimTimeSliderController.h
--- a/Engine/Source/Editor/Persona/Private/AnimTimeline/AnimTimeSliderController.h
+++ b/Engine/Source/Editor/Persona/Private/AnimTimeline/AnimTimeSliderController.h
@@ -6,7 +6,7 @@
 #include "Input/CursorReply.h"
 #include "Input/Reply.h"
 #include "Widgets/SWidget.h"
-#include "ITimeSlider.h"
+#include "TimeSliderArgs.h"
 #include "AnimTimeline/SAnimTimeline.h"
 
 class FSlateWindowElementList;
diff --git a/Engine/Source/Editor/Persona/Private/AnimTimeline/SAnimTimeline.cpp b/Engine/Source/Editor/Persona/Private/AnimTimeline/SAnimTimeline.cpp
--- a/Engine/Source/Editor/Persona/Private/AnimTimeline/SAnimTimeline.cpp
+++ b/Engine/Source/Editor/Persona/Private/AnimTimeline/SAnimTimeline.cpp
@@ -41,6 +41,7 @@
 #include "Widgets/Layout/SScrollBox.h"
 #include "AnimTimeline/AnimTimeSliderController.h"
 #include "Framework/Commands/GenericCommands.h"
+#include "MovieSceneFwd.h"
 
 #define LOCTEXT_NAMESPACE "SAnimTimeline"
 
diff --git a/Engine/Source/Editor/Persona/Private/AnimTimeline/SAnimTimeline.h b/Engine/Source/Editor/Persona/Private/AnimTimeline/SAnimTimeline.h
--- a/Engine/Source/Editor/Persona/Private/AnimTimeline/SAnimTimeline.h
+++ b/Engine/Source/Editor/Persona/Private/AnimTimeline/SAnimTimeline.h
@@ -2,6 +2,7 @@
 
 #pragma once
 
+#include "AnimatedRange.h"
 #include "CoreTypes.h"
 #include "Templates/SharedPointer.h"
 #include "Widgets/DeclarativeSyntaxSupport.h"
diff --git a/Engine/Source/Editor/Persona/Private/SAnimSequenceCurveEditor.cpp b/Engine/Source/Editor/Persona/Private/SAnimSequenceCurveEditor.cpp
--- a/Engine/Source/Editor/Persona/Private/SAnimSequenceCurveEditor.cpp
+++ b/Engine/Source/Editor/Persona/Private/SAnimSequenceCurveEditor.cpp
@@ -1,6 +1,7 @@
 // Copyright Epic Games, Inc. All Rights Reserved.
 
 #include "SAnimSequenceCurveEditor.h"
+#include "AnimatedRange.h"
 #include "CurveEditor.h"
 #include "RichCurveEditorModel.h"
 #include "Animation/AnimSequenceBase.h"
diff --git a/Engine/Source/Editor/Sequencer/Private/IKeyArea.cpp b/Engine/Source/Editor/Sequencer/Private/IKeyArea.cpp
--- a/Engine/Source/Editor/Sequencer/Private/IKeyArea.cpp
+++ b/Engine/Source/Editor/Sequencer/Private/IKeyArea.cpp
@@ -12,6 +12,7 @@
 #include "CurveModel.h"
 #include "ISequencer.h"
 #include "ISequencerSection.h"
+#include "MovieScene.h"
 #include "MovieSceneSequence.h"
 #include "CurveEditorSettings.h"
 
diff --git a/Engine/Source/Editor/Sequencer/Private/ISequencer.cpp b/Engine/Source/Editor/Sequencer/Private/ISequencer.cpp
--- a/Engine/Source/Editor/Sequencer/Private/ISequencer.cpp
+++ b/Engine/Source/Editor/Sequencer/Private/ISequencer.cpp
@@ -2,9 +2,16 @@
 
 #include "ISequencer.h"
 
+#include "AnimatedRange.h"
 #include "Misc/AssertionMacros.h"
 #include "MovieScene.h"
 #include "MovieSceneSequence.h"
+#include "ITimeSlider.h"
+
+FAnimatedRange ISequencer::GetViewRange() const
+{
+	return FAnimatedRange();
+}
 
 FFrameRate ISequencer::GetRootTickResolution() const
 {
diff --git a/Engine/Source/Editor/Sequencer/Private/MVVM/FolderModelStorageExtension.cpp b/Engine/Source/Editor/Sequencer/Private/MVVM/FolderModelStorageExtension.cpp
--- a/Engine/Source/Editor/Sequencer/Private/MVVM/FolderModelStorageExtension.cpp
+++ b/Engine/Source/Editor/Sequencer/Private/MVVM/FolderModelStorageExtension.cpp
@@ -7,6 +7,7 @@
 #include "MVVM/ViewModels/FolderModel.h"
 #include "MVVM/ViewModels/TrackModel.h"
 
+#include "MovieScene.h"
 #include "MovieSceneSequence.h"
 #include "MovieSceneFolder.h"
 
diff --git a/Engine/Source/Editor/Sequencer/Private/MVVM/TrackModelStorageExtension.cpp b/Engine/Source/Editor/Sequencer/Private/MVVM/TrackModelStorageExtension.cpp
--- a/Engine/Source/Editor/Sequencer/Private/MVVM/TrackModelStorageExtension.cpp
+++ b/Engine/Source/Editor/Sequencer/Private/MVVM/TrackModelStorageExtension.cpp
@@ -4,6 +4,7 @@
 #include "MVVM/ViewModels/TrackModel.h"
 #include "MVVM/ViewModels/SequenceModel.h"
 #include "MVVM/ViewModels/ViewModel.h"
+#include "MovieScene.h"
 #include "MovieSceneSequence.h"
 #include "MovieSceneFolder.h"
 
diff --git a/Engine/Source/Editor/Sequencer/Private/MVVM/ViewModels/FolderModel.cpp b/Engine/Source/Editor/Sequencer/Private/MVVM/ViewModels/FolderModel.cpp
--- a/Engine/Source/Editor/Sequencer/Private/MVVM/ViewModels/FolderModel.cpp
+++ b/Engine/Source/Editor/Sequencer/Private/MVVM/ViewModels/FolderModel.cpp
@@ -25,6 +25,7 @@
 #include "Widgets/Colors/SColorPicker.h"
 #include "SequencerOutlinerItemDragDropOp.h"
 
+#include "MovieScene.h"
 #include "MovieSceneBinding.h"
 #include "MovieSceneFolder.h"
 
diff --git a/Engine/Source/Editor/Sequencer/Private/MVVM/ViewModels/OutlinerItemModel.cpp b/Engine/Source/Editor/Sequencer/Private/MVVM/ViewModels/OutlinerItemModel.cpp
--- a/Engine/Source/Editor/Sequencer/Private/MVVM/ViewModels/OutlinerItemModel.cpp
+++ b/Engine/Source/Editor/Sequencer/Private/MVVM/ViewModels/OutlinerItemModel.cpp
@@ -15,6 +15,7 @@
 #include "Framework/Commands/UIAction.h"
 #include "Framework/MultiBox/MultiBoxBuilder.h"
 #include "ISettingsModule.h"
+#include "MovieScene.h"
 #include "MovieSceneSequence.h"
 #include "ScopedTransaction.h"
 #include "Sequencer.h"
diff --git a/Engine/Source/Editor/Sequencer/Private/MVVM/ViewModels/SequenceModel.cpp b/Engine/Source/Editor/Sequencer/Private/MVVM/ViewModels/SequenceModel.cpp
--- a/Engine/Source/Editor/Sequencer/Private/MVVM/ViewModels/SequenceModel.cpp
+++ b/Engine/Source/Editor/Sequencer/Private/MVVM/ViewModels/SequenceModel.cpp
@@ -11,6 +11,7 @@
 #include "MVVM/ViewModels/ViewModel.h"
 #include "MVVM/ViewModels/ViewModelIterators.h"
 
+#include "MovieScene.h"
 #include "MovieSceneFolder.h"
 #include "MovieSceneSequence.h"
 
diff --git a/Engine/Source/Editor/Sequencer/Private/MVVM/ViewModels/SequencerTrackAreaViewModel.cpp b/Engine/Source/Editor/Sequencer/Private/MVVM/ViewModels/SequencerTrackAreaViewModel.cpp
--- a/Engine/Source/Editor/Sequencer/Private/MVVM/ViewModels/SequencerTrackAreaViewModel.cpp
+++ b/Engine/Source/Editor/Sequencer/Private/MVVM/ViewModels/SequencerTrackAreaViewModel.cpp
@@ -3,6 +3,7 @@
 #include "MVVM/ViewModels/SequencerTrackAreaViewModel.h"
 #include "MVVM/ViewModels/SequencerEditorViewModel.h"
 #include "MVVM/ViewModels/SequenceModel.h"
+#include "AnimatedRange.h"
 #include "Sequencer.h"
 
 #include "Framework/MultiBox/MultiBoxBuilder.h"
diff --git a/Engine/Source/Editor/Sequencer/Private/MVVM/ViewModels/TrackModel.cpp b/Engine/Source/Editor/Sequencer/Private/MVVM/ViewModels/TrackModel.cpp
--- a/Engine/Source/Editor/Sequencer/Private/MVVM/ViewModels/TrackModel.cpp
+++ b/Engine/Source/Editor/Sequencer/Private/MVVM/ViewModels/TrackModel.cpp
@@ -17,6 +17,7 @@
 #include "MVVM/Views/SOutlinerTrackView.h"
 #include "MVVM/ViewModels/SequencerEditorViewModel.h"
 
+#include "MovieScene.h"
 #include "MovieSceneFolder.h"
 #include "MovieSceneTrack.h"
 #include "MovieSceneSection.h"
diff --git a/Engine/Source/Editor/Sequencer/Private/MVVM/Views/SChannelView.cpp b/Engine/Source/Editor/Sequencer/Private/MVVM/Views/SChannelView.cpp
--- a/Engine/Source/Editor/Sequencer/Private/MVVM/Views/SChannelView.cpp
+++ b/Engine/Source/Editor/Sequencer/Private/MVVM/Views/SChannelView.cpp
@@ -8,6 +8,7 @@
 #include "MVVM/Views/STrackAreaView.h"
 
 #include "IKeyArea.h"
+#include "MovieScene.h"
 #include "Sequencer.h"
 #include "SequencerHotspots.h"
 #include "ISequencerEditTool.h"
diff --git a/Engine/Source/Editor/Sequencer/Private/SSequencer.h b/Engine/Source/Editor/Sequencer/Private/SSequencer.h
--- a/Engine/Source/Editor/Sequencer/Private/SSequencer.h
+++ b/Engine/Source/Editor/Sequencer/Private/SSequencer.h
@@ -3,6 +3,7 @@
 #pragma once
 
 #include "CoreMinimal.h"
+#include "AnimatedRange.h"
 #include "Layout/Visibility.h"
 #include "Input/Reply.h"
 #include "Widgets/SWidget.h"
@@ -19,6 +20,7 @@
 #include "Widgets/Input/SSpinBox.h"
 #include "Widgets/Text/STextBlock.h"
 #include "Sequencer.h"
+#include "SequencerWidgetsDelegates.h"
 
 class FActorDragDropOp;
 class FFolderDragDropOp;
diff --git a/Engine/Source/Editor/Sequencer/Private/SSequencerSection.cpp b/Engine/Source/Editor/Sequencer/Private/SSequencerSection.cpp
--- a/Engine/Source/Editor/Sequencer/Private/SSequencerSection.cpp
+++ b/Engine/Source/Editor/Sequencer/Private/SSequencerSection.cpp
@@ -12,6 +12,7 @@
 #include "MVVM/Views/STrackLane.h"
 #include "MVVM/Views/SChannelView.h"
 #include "MVVM/Extensions/IObjectBindingExtension.h"
+#include "AnimatedRange.h"
 #include "Rendering/DrawElements.h"
 #include "Styling/AppStyle.h"
 #include "Styling/StyleColors.h"
diff --git a/Engine/Source/Editor/Sequencer/Private/Sequencer.cpp b/Engine/Source/Editor/Sequencer/Private/Sequencer.cpp
--- a/Engine/Source/Editor/Sequencer/Private/Sequencer.cpp
+++ b/Engine/Source/Editor/Sequencer/Private/Sequencer.cpp
@@ -89,6 +89,7 @@
 #include "Sections/MovieSceneBoolSection.h"
 #include "Sections/MovieScene3DTransformSection.h"
 #include "Sections/MovieSceneSubSection.h"
+#include "TimeSliderArgs.h"
 #include "Tracks/MovieScene3DTransformTrack.h"
 #include "Tracks/MovieSceneCinematicShotTrack.h"
 #include "Tracks/MovieScenePropertyTrack.h"
diff --git a/Engine/Source/Editor/Sequencer/Private/Sequencer.h b/Engine/Source/Editor/Sequencer/Private/Sequencer.h
--- a/Engine/Source/Editor/Sequencer/Private/Sequencer.h
+++ b/Engine/Source/Editor/Sequencer/Private/Sequencer.h
@@ -12,6 +12,7 @@
 #include "Widgets/SWidget.h"
 #include "SequencerNodeTree.h"
 #include "UObject/GCObject.h"
+#include "MovieSceneMarkedFrame.h"
 #include "MovieSceneSequenceID.h"
 #include "IMovieScenePlayer.h"
 #include "ITimeSlider.h"
@@ -64,6 +65,8 @@ class USequencerSettings;
 class UMovieSceneCopyableBinding;
 class UMovieSceneCompiledDataManager;
 class UMovieSceneCopyableTrack;
+class UMovieSceneNodeGroup;
+
 struct FMovieSceneTimeController;
 struct FMovieSceneSequencePlaybackParams;
 struct FMovieScenePossessable;
@@ -72,8 +75,10 @@ struct FKeyAttributes;
 class USequencerTimeChangeUndoRedoProxy;
 struct FNotificationInfo;
 struct FEditorViewportViewModifierParams;
+struct FMovieSceneMarkedFrame;
 
 enum class EMapChangeType : uint8;
+enum class ENearestKeyOption : uint8;
 
 namespace UE
 {
diff --git a/Engine/Source/Editor/Sequencer/Private/SequencerUtilities.cpp b/Engine/Source/Editor/Sequencer/Private/SequencerUtilities.cpp
--- a/Engine/Source/Editor/Sequencer/Private/SequencerUtilities.cpp
+++ b/Engine/Source/Editor/Sequencer/Private/SequencerUtilities.cpp
@@ -1,6 +1,7 @@
 // Copyright Epic Games, Inc. All Rights Reserved.
 
 #include "SequencerUtilities.h"
+#include "AnimatedRange.h"
 #include "CineCameraActor.h"
 #include "CameraRig_Rail.h"
 #include "CameraRig_Crane.h"
diff --git a/Engine/Source/Editor/Sequencer/Private/Tools/EditToolDragOperations.cpp b/Engine/Source/Editor/Sequencer/Private/Tools/EditToolDragOperations.cpp
--- a/Engine/Source/Editor/Sequencer/Private/Tools/EditToolDragOperations.cpp
+++ b/Engine/Source/Editor/Sequencer/Private/Tools/EditToolDragOperations.cpp
@@ -18,6 +18,7 @@
 #include "MVVM/ViewModels/TrackAreaViewModel.h"
 #include "MVVM/ViewModels/VirtualTrackArea.h"
 #include "Algo/AllOf.h"
+#include "MovieScene.h"
 #include "MovieSceneTimeHelpers.h"
 #include "Modules/ModuleManager.h"
 #include "Channels/MovieSceneChannel.h"
diff --git a/Engine/Source/Editor/Sequencer/Private/Widgets/CurveEditor/SSequencerCurveEditor.cpp b/Engine/Source/Editor/Sequencer/Private/Widgets/CurveEditor/SSequencerCurveEditor.cpp
--- a/Engine/Source/Editor/Sequencer/Private/Widgets/CurveEditor/SSequencerCurveEditor.cpp
+++ b/Engine/Source/Editor/Sequencer/Private/Widgets/CurveEditor/SSequencerCurveEditor.cpp
@@ -2,6 +2,7 @@
 
 #include "SSequencerCurveEditor.h"
 #include "Framework/MultiBox/MultiBoxBuilder.h"
+#include "MovieScene.h"
 #include "MovieSceneCommonHelpers.h"
 #include "SCurveEditorPanel.h"
 #include "Sequencer.h"
diff --git a/Engine/Source/Editor/Sequencer/Public/ISequencer.h b/Engine/Source/Editor/Sequencer/Public/ISequencer.h
--- a/Engine/Source/Editor/Sequencer/Public/ISequencer.h
+++ b/Engine/Source/Editor/Sequencer/Public/ISequencer.h
@@ -8,7 +8,7 @@
 #include "Containers/UnrealString.h"
 #include "CoreMinimal.h"
 #include "Delegates/Delegate.h"
-#include "ITimeSlider.h"
+#include "ViewRangeInterpolation.h"
 #include "Evaluation/MovieSceneSequenceTransform.h"
 #include "HAL/Platform.h"
 #include "IMovieScenePlayer.h"
@@ -32,6 +32,10 @@
 #include "Widgets/Input/NumericTypeInterface.h"
 #include "Widgets/SWidget.h"
 
+#if UE_ENABLE_INCLUDE_ORDER_DEPRECATED_IN_5_3
+#include "ITimeSlider.h"
+#endif
+
 #include "ISequencer.generated.h"
 
 class AActor;
@@ -67,6 +71,7 @@ enum class EMapChangeType : uint8;
 class FCurveEditor;
 class FCurveModel;
 class IToolkitHost;
+struct FAnimatedRange;
 struct FMovieSceneChannelMetaData;
 struct FMovieSceneSequencePlaybackParams;
 
@@ -436,10 +441,7 @@ public:
 	virtual void ResetTimeController() = 0;
 
 	/** @return The current view range */
-	virtual FAnimatedRange GetViewRange() const
-	{
-		return FAnimatedRange();
-	}
+	virtual FAnimatedRange GetViewRange() const;
 
 	/**
 	 * Set the view range, growing the working range to accomodate, if necessary
diff --git a/Engine/Source/Editor/Sequencer/Public/SequencerTimeSliderController.h b/Engine/Source/Editor/Sequencer/Public/SequencerTimeSliderController.h
--- a/Engine/Source/Editor/Sequencer/Public/SequencerTimeSliderController.h
+++ b/Engine/Source/Editor/Sequencer/Public/SequencerTimeSliderController.h
@@ -8,6 +8,7 @@
 #include "Widgets/SWidget.h"
 #include "ITimeSlider.h"
 #include "ISequencerModule.h"
+#include "TimeSliderArgs.h"
 
 class FSlateWindowElementList;
 struct FContextMenuSuppressor;
diff --git a/Engine/Source/Editor/SequencerWidgets/Private/ITimeSlider.cpp b/Engine/Source/Editor/SequencerWidgets/Private/ITimeSlider.cpp
new file mode 100644
--- /dev/null
+++ b/Engine/Source/Editor/SequencerWidgets/Private/ITimeSlider.cpp
@@ -0,0 +1,50 @@
+// Copyright Epic Games, Inc. All Rights Reserved.
+
+#include "ITimeSlider.h"
+#include "AnimatedRange.h"
+#include "Misc/FrameTime.h"
+
+FAnimatedRange ITimeSliderController::GetViewRange() const
+{
+	return FAnimatedRange();
+}
+
+FAnimatedRange ITimeSliderController::GetClampRange() const
+{
+	return FAnimatedRange();
+}
+
+TRange<FFrameNumber> ITimeSliderController::GetPlayRange() const
+{
+	return TRange<FFrameNumber>();
+}
+
+TRange<FFrameNumber> ITimeSliderController::GetSelectionRange() const
+{
+	return TRange<FFrameNumber>();
+}
+
+FFrameTime ITimeSliderController::GetScrubPosition() const
+{
+	return FFrameTime();
+}
+
+void ITimeSliderController::SetScrubPosition(FFrameTime InTime, bool bEvaluate)
+{
+}
+
+void ITimeSliderController::SetViewRange(double NewRangeMin, double NewRangeMax, EViewRangeInterpolation Interpolation)
+{
+}
+
+void ITimeSliderController::SetClampRange(double NewRangeMin, double NewRangeMax)
+{
+}
+
+void ITimeSliderController::SetPlayRange(FFrameNumber RangeStart, int32 RangeDuration)
+{
+}
+
+void ITimeSliderController::SetSelectionRange(const TRange<FFrameNumber>& NewRange)
+{
+}
diff --git a/Engine/Source/Editor/SequencerWidgets/Private/STimeRange.cpp b/Engine/Source/Editor/SequencerWidgets/Private/STimeRange.cpp
--- a/Engine/Source/Editor/SequencerWidgets/Private/STimeRange.cpp
+++ b/Engine/Source/Editor/SequencerWidgets/Private/STimeRange.cpp
@@ -2,6 +2,7 @@
 
 #include "STimeRange.h"
 
+#include "AnimatedRange.h"
 #include "HAL/Platform.h"
 #include "HAL/PlatformCrt.h"
 #include "Internationalization/Internationalization.h"
diff --git a/Engine/Source/Editor/SequencerWidgets/Private/STimeRangeSlider.cpp b/Engine/Source/Editor/SequencerWidgets/Private/STimeRangeSlider.cpp
--- a/Engine/Source/Editor/SequencerWidgets/Private/STimeRangeSlider.cpp
+++ b/Engine/Source/Editor/SequencerWidgets/Private/STimeRangeSlider.cpp
@@ -2,6 +2,7 @@
 
 #include "STimeRangeSlider.h"
 
+#include "AnimatedRange.h"
 #include "ITimeSlider.h"
 #include "Input/Events.h"
 #include "Layout/Geometry.h"
diff --git a/Engine/Source/Editor/SequencerWidgets/Public/AnimatedRange.h b/Engine/Source/Editor/SequencerWidgets/Public/AnimatedRange.h
new file mode 100644
--- /dev/null
+++ b/Engine/Source/Editor/SequencerWidgets/Public/AnimatedRange.h
@@ -0,0 +1,41 @@
+// Copyright Epic Games, Inc. All Rights Reserved.
+
+#pragma once
+
+#include "Math/Range.h"
+#include "Misc/Attribute.h"
+#include "Misc/Optional.h"
+
+/** Structure used to wrap up a range, and an optional animation target */
+struct FAnimatedRange : public TRange<double>
+{
+	/** Default Construction */
+	FAnimatedRange() : TRange() {}
+	/** Construction from a lower and upper bound */
+	FAnimatedRange(double LowerBound, double UpperBound) : TRange(LowerBound, UpperBound) {}
+	/** Copy-construction from simple range */
+	FAnimatedRange(const TRange<double>& InRange) : TRange(InRange) {}
+
+	/** Helper function to wrap an attribute to an animated range with a non-animated one */
+	static TAttribute<TRange<double>> WrapAttribute(const TAttribute<FAnimatedRange>& InAttribute)
+	{
+		typedef TAttribute<TRange<double>> Attr;
+		return Attr::Create(Attr::FGetter::CreateLambda([=]() { return InAttribute.Get(); }));
+	}
+
+	/** Helper function to wrap an attribute to a non-animated range with an animated one */
+	static TAttribute<FAnimatedRange> WrapAttribute(const TAttribute<TRange<double>>& InAttribute)
+	{
+		typedef TAttribute<FAnimatedRange> Attr;
+		return Attr::Create(Attr::FGetter::CreateLambda([=]() { return InAttribute.Get(); }));
+	}
+
+	/** Get the current animation target, or the whole view range when not animating */
+	const TRange<double>& GetAnimationTarget() const
+	{
+		return AnimationTarget.IsSet() ? AnimationTarget.GetValue() : *this;
+	}
+
+	/** The animation target, if animating */
+	TOptional<TRange<double>> AnimationTarget;
+};
diff --git a/Engine/Source/Editor/SequencerWidgets/Public/ITimeSlider.h b/Engine/Source/Editor/SequencerWidgets/Public/ITimeSlider.h
--- a/Engine/Source/Editor/SequencerWidgets/Public/ITimeSlider.h
+++ b/Engine/Source/Editor/SequencerWidgets/Public/ITimeSlider.h
@@ -2,6 +2,11 @@
 
 #pragma once
 
+#include "ISequencerInputHandler.h"
+#include "ViewRangeInterpolation.h"
+#include "Widgets/SCompoundWidget.h"
+
+#if UE_ENABLE_INCLUDE_ORDER_DEPRECATED_IN_5_3
 #include "CoreMinimal.h"
 #include "Misc/Attribute.h"
 #include "Input/CursorReply.h"
@@ -13,206 +18,18 @@
 #include "IMovieScenePlayer.h"
 #include "MovieScene.h"
 
-class FSlateWindowElementList;
-
-/** Enum specifying how to interpolate to a new view range */
-enum class EViewRangeInterpolation
-{
-	/** Use an externally defined animated interpolation */
-	Animated,
-	/** Set the view range immediately */
-	Immediate,
-};
-
-/** Enum specifying how to find the nearest key */
-enum class ENearestKeyOption : uint8
-{
-	NKO_None = 0x00,
-
-	/* Search keys */
-	NKO_SearchKeys = 0x01,
-
-	/* Search sections */
-	NKO_SearchSections = 0x02,
-
-	/* Search markers */
-	NKO_SearchMarkers = 0x04,
-
-	/** Search all tracks */
-	NKO_SearchAllTracks = 0x08
-};
-ENUM_CLASS_FLAGS(ENearestKeyOption);
-
-DECLARE_DELEGATE_ThreeParams( FOnScrubPositionChanged, FFrameTime, bool, bool )
-DECLARE_DELEGATE_TwoParams( FOnViewRangeChanged, TRange<double>, EViewRangeInterpolation )
-DECLARE_DELEGATE_OneParam( FOnTimeRangeChanged, TRange<double> )
-DECLARE_DELEGATE_OneParam( FOnFrameRangeChanged, TRange<FFrameNumber> )
-DECLARE_DELEGATE_TwoParams(FOnSetMarkedFrame, int32, FFrameNumber)
-DECLARE_DELEGATE_OneParam(FOnAddMarkedFrame, FFrameNumber)
-DECLARE_DELEGATE_OneParam(FOnDeleteMarkedFrame, int32)
-DECLARE_DELEGATE_RetVal_TwoParams( FFrameNumber, FOnGetNearestKey, FFrameTime, ENearestKeyOption )
-DECLARE_DELEGATE_OneParam(FOnScrubPositionParentChanged, FMovieSceneSequenceID)
-
-/** Structure used to wrap up a range, and an optional animation target */
-struct FAnimatedRange : public TRange<double>
-{
-	/** Default Construction */
-	FAnimatedRange() : TRange() {}
-	/** Construction from a lower and upper bound */
-	FAnimatedRange( double LowerBound, double UpperBound ) : TRange( LowerBound, UpperBound ) {}
-	/** Copy-construction from simple range */
-	FAnimatedRange( const TRange<double>& InRange ) : TRange(InRange) {}
-
-	/** Helper function to wrap an attribute to an animated range with a non-animated one */
-	static TAttribute<TRange<double>> WrapAttribute( const TAttribute<FAnimatedRange>& InAttribute )
-	{
-		typedef TAttribute<TRange<double>> Attr;
-		return Attr::Create(Attr::FGetter::CreateLambda([=](){ return InAttribute.Get(); }));
-	}
-
-	/** Helper function to wrap an attribute to a non-animated range with an animated one */
-	static TAttribute<FAnimatedRange> WrapAttribute( const TAttribute<TRange<double>>& InAttribute )
-	{
-		typedef TAttribute<FAnimatedRange> Attr;
-		return Attr::Create(Attr::FGetter::CreateLambda([=](){ return InAttribute.Get(); }));
-	}
-
-	/** Get the current animation target, or the whole view range when not animating */
-	const TRange<double>& GetAnimationTarget() const
-	{
-		return AnimationTarget.IsSet() ? AnimationTarget.GetValue() : *this;
-	}
-	
-	/** The animation target, if animating */
-	TOptional<TRange<double>> AnimationTarget;
-};
-
-struct FTimeSliderArgs
-{
-	FTimeSliderArgs()
-		: ScrubPosition(0)
-		, ViewRange( FAnimatedRange(0.0f, 5.0f) )
-		, ClampRange( FAnimatedRange(-FLT_MAX/2.f, FLT_MAX/2.f) )
-		, AllowZoom(true)
-	{}
-
-	/** The scrub position */
-	TAttribute<FFrameTime> ScrubPosition;
-
-	/** The scrub position text */
-	TAttribute<FString> ScrubPositionText;
-
-	/** The parent sequence that the scrub position display text is relative to */
-	TAttribute<FMovieSceneSequenceID> ScrubPositionParent;
-
-	/** Called when the scrub position parent sequence is changed */
-	FOnScrubPositionParentChanged OnScrubPositionParentChanged;
-
-	/** Attribute for the parent sequence chain of the current sequence */
-	TAttribute<TArray<FMovieSceneSequenceID>> ScrubPositionParentChain;
+#include "AnimatedRange.h"
+#include "TimeSliderArgs.h"
+#endif
 
-	/** View time range */
-	TAttribute< FAnimatedRange > ViewRange;
-
-	/** Clamp time range */
-	TAttribute< FAnimatedRange > ClampRange;
-
-	/** Called when the scrub position changes */
-	FOnScrubPositionChanged OnScrubPositionChanged;
-
-	/** Called right before the scrubber begins to move */
-	FSimpleDelegate OnBeginScrubberMovement;
-
-	/** Called right after the scrubber handle is released by the user */
-	FSimpleDelegate OnEndScrubberMovement;
-
-	/** Called when the view range changes */
-	FOnViewRangeChanged OnViewRangeChanged;
-
-	/** Called when the clamp range changes */
-	FOnTimeRangeChanged OnClampRangeChanged;
-
-	/** Delegate that is called when getting the nearest key */
-	FOnGetNearestKey OnGetNearestKey;
-
-	/** Attribute defining the active sub-sequence range for this controller */
-	TAttribute<TOptional<TRange<FFrameNumber>>> SubSequenceRange;
-
-	/** Attribute defining the playback range for this controller */
-	TAttribute<TRange<FFrameNumber>> PlaybackRange;
-
-	/** Attribute for the current sequence's display rate */
-	TAttribute<FFrameRate> DisplayRate;
-
-	/** Attribute for the current sequence's tick resolution */
-	TAttribute<FFrameRate> TickResolution;
-
-	/** Delegate that is called when the playback range wants to change */
-	FOnFrameRangeChanged OnPlaybackRangeChanged;
-
-	/** Called right before the playback range starts to be dragged */
-	FSimpleDelegate OnPlaybackRangeBeginDrag;
-
-	/** Called right after the playback range has finished being dragged */
-	FSimpleDelegate OnPlaybackRangeEndDrag;
-
-	/** Attribute defining the selection range for this controller */
-	TAttribute<TRange<FFrameNumber>> SelectionRange;
-
-	/** Delegate that is called when the selection range wants to change */
-	FOnFrameRangeChanged OnSelectionRangeChanged;
-
-	/** Called right before the selection range starts to be dragged */
-	FSimpleDelegate OnSelectionRangeBeginDrag;
-
-	/** Called right after the selection range has finished being dragged */
-	FSimpleDelegate OnSelectionRangeEndDrag;
-
-	/** Called right before a mark starts to be dragged */
-	FSimpleDelegate OnMarkBeginDrag;
-
-	/** Called right after a mark has finished being dragged */
-	FSimpleDelegate OnMarkEndDrag;
-
-	/** Attribute for the current sequence's vertical frames */
-	TAttribute<TSet<FFrameNumber>> VerticalFrames;
-
-	/** Attribute for the current sequence's marked frames */
-	TAttribute<TArray<FMovieSceneMarkedFrame>> MarkedFrames;
-
-	/** Attribute for the marked frames that might need to be shown, but do not belong to the current sequence*/
-	TAttribute<TArray<FMovieSceneMarkedFrame>> GlobalMarkedFrames;
-
-	/** Called when the marked frame needs to be set */
-	FOnSetMarkedFrame OnSetMarkedFrame;
-
-	/** Called when a marked frame is added */
-	FOnAddMarkedFrame OnAddMarkedFrame;
-
-	/** Called when a marked frame is deleted */
-	FOnDeleteMarkedFrame OnDeleteMarkedFrame;
-
-	/** Called when all marked frames should be deleted */
-	FSimpleDelegate OnDeleteAllMarkedFrames;
-
-	/** Round the scrub position to an integer during playback */
-	TAttribute<EMovieScenePlayerStatus::Type> PlaybackStatus;
-
-	/** Attribute defining whether the playback range is locked */
-	TAttribute<bool> IsPlaybackRangeLocked;
-
-	/** Attribute defining the time snap interval */
-	TAttribute<float> TimeSnapInterval;
-
-	/** Called when toggling the playback range lock */
-	FSimpleDelegate OnTogglePlaybackRangeLocked;
-
-	/** If we are allowed to zoom */
-	bool AllowZoom;
+class FSlateWindowElementList;
+struct FSlateBrush;
 
-	/** Numeric Type interface for converting between frame numbers and display formats. */
-	TSharedPtr<INumericTypeInterface<double>> NumericTypeInterface;
-};
+struct FGeometry;
+struct FFrameRate;
+struct FAnimatedRange;
+class FSlateRect;
+class FWidgetStyle;
 
 struct FPaintPlaybackRangeArgs
 {
@@ -249,7 +66,6 @@ struct FPaintViewAreaArgs
 	TOptional<FPaintPlaybackRangeArgs> PlaybackRangeArgs;
 };
 
-
 class ITimeSliderController : public ISequencerInputHandler
 {
 public:
@@ -289,22 +105,22 @@ public:
 	virtual FFrameRate GetTickResolution() const = 0;
 
 	/** Get the current view range for this controller */
-	virtual FAnimatedRange GetViewRange() const { return FAnimatedRange(); }
+	SEQUENCERWIDGETS_API virtual FAnimatedRange GetViewRange() const;
 
 	/** Get the current clamp range for this controller */
-	virtual FAnimatedRange GetClampRange() const { return FAnimatedRange(); }
+	SEQUENCERWIDGETS_API virtual FAnimatedRange GetClampRange() const;
 
 	/** Get the current play range for this controller */
-	virtual TRange<FFrameNumber> GetPlayRange() const { return TRange<FFrameNumber>(); }
+	SEQUENCERWIDGETS_API virtual TRange<FFrameNumber> GetPlayRange() const;
 
 	/** Get the current selection range for this controller */
-	virtual TRange<FFrameNumber> GetSelectionRange() const { return TRange<FFrameNumber>(); }
+	SEQUENCERWIDGETS_API virtual TRange<FFrameNumber> GetSelectionRange() const;
 
 	/** Get the current time for the Scrub handle which indicates what range is being evaluated. */
-	virtual FFrameTime GetScrubPosition() const { return FFrameTime(); }
+	SEQUENCERWIDGETS_API virtual FFrameTime GetScrubPosition() const;
 
 	/** Set the current time for the Scrub handle which indicates what range is being evaluated. */
-	virtual void SetScrubPosition(FFrameTime InTime, bool bEvaluate) {}
+	SEQUENCERWIDGETS_API virtual void SetScrubPosition(FFrameTime InTime, bool bEvaluate);
 
 	/**
 	 * Set a new range based on a min, max and an interpolation mode
@@ -313,7 +129,7 @@ public:
 	 * @param NewRangeMax		The new upper bound of the range
 	 * @param Interpolation		How to set the new range (either immediately, or animated)
 	 */
-	virtual void SetViewRange( double NewRangeMin, double NewRangeMax, EViewRangeInterpolation Interpolation ) {}
+	SEQUENCERWIDGETS_API virtual void SetViewRange(double NewRangeMin, double NewRangeMax, EViewRangeInterpolation Interpolation);
 
 	/**
 	 * Set a new clamp range based on a min, max
@@ -321,7 +137,7 @@ public:
 	 * @param NewRangeMin		The new lower bound of the clamp range
 	 * @param NewRangeMax		The new upper bound of the clamp range
 	 */
-	virtual void SetClampRange( double NewRangeMin, double NewRangeMax) {}
+	SEQUENCERWIDGETS_API virtual void SetClampRange(double NewRangeMin, double NewRangeMax);
 
 	/**
 	 * Set a new playback range based on a min, max
@@ -329,14 +145,14 @@ public:
 	 * @param RangeStart		The new lower bound of the playback range
 	 * @param RangeDuration		The total number of frames that we play for
 	 */
-	virtual void SetPlayRange( FFrameNumber RangeStart, int32 RangeDuration ) {}
+	SEQUENCERWIDGETS_API virtual void SetPlayRange(FFrameNumber RangeStart, int32 RangeDuration);
 
 	/**
 	 * Set a new selection range
 	 * 
 	 * @param NewRange		The new selection range
 	 */
-	virtual void SetSelectionRange(const TRange<FFrameNumber>& NewRange) {}
+	SEQUENCERWIDGETS_API virtual void SetSelectionRange(const TRange<FFrameNumber>& NewRange);
 };
 
 /**
diff --git a/Engine/Source/Editor/SequencerWidgets/Public/SequencerWidgetsDelegates.h b/Engine/Source/Editor/SequencerWidgets/Public/SequencerWidgetsDelegates.h
new file mode 100644
--- /dev/null
+++ b/Engine/Source/Editor/SequencerWidgets/Public/SequencerWidgetsDelegates.h
@@ -0,0 +1,22 @@
+// Copyright Epic Games, Inc. All Rights Reserved.
+
+#pragma once
+
+#include "Delegates/Delegate.h"
+#include "Math/Range.h"
+#include "ViewRangeInterpolation.h"
+
+struct FFrameTime;
+struct FFrameNumber;
+struct FMovieSceneSequenceID;
+enum class ENearestKeyOption : uint8;
+
+DECLARE_DELEGATE_ThreeParams(FOnScrubPositionChanged, FFrameTime, bool, bool)
+DECLARE_DELEGATE_TwoParams(FOnViewRangeChanged, TRange<double>, EViewRangeInterpolation)
+DECLARE_DELEGATE_OneParam(FOnTimeRangeChanged, TRange<double>)
+DECLARE_DELEGATE_OneParam(FOnFrameRangeChanged, TRange<FFrameNumber>)
+DECLARE_DELEGATE_TwoParams(FOnSetMarkedFrame, int32, FFrameNumber)
+DECLARE_DELEGATE_OneParam(FOnAddMarkedFrame, FFrameNumber)
+DECLARE_DELEGATE_OneParam(FOnDeleteMarkedFrame, int32)
+DECLARE_DELEGATE_RetVal_TwoParams(FFrameNumber, FOnGetNearestKey, FFrameTime, ENearestKeyOption)
+DECLARE_DELEGATE_OneParam(FOnScrubPositionParentChanged, FMovieSceneSequenceID)
diff --git a/Engine/Source/Editor/SequencerWidgets/Public/TimeSliderArgs.h b/Engine/Source/Editor/SequencerWidgets/Public/TimeSliderArgs.h
new file mode 100644
--- /dev/null
+++ b/Engine/Source/Editor/SequencerWidgets/Public/TimeSliderArgs.h
@@ -0,0 +1,165 @@
+// Copyright Epic Games, Inc. All Rights Reserved.
+
+#pragma once
+
+#include "AnimatedRange.h"
+#include "Delegates/Delegate.h"
+#include "Misc/EnumClassFlags.h"
+#include "Misc/FrameNumber.h"
+#include "Misc/FrameRate.h"
+#include "Misc/FrameTime.h"
+#include "MovieSceneMarkedFrame.h"
+#include "MovieSceneSequenceID.h"
+#include "SequencerWidgetsDelegates.h"
+#include "Widgets/Input/NumericTypeInterface.h"
+
+namespace EMovieScenePlayerStatus
+{
+	enum Type : int;
+}
+
+/** Enum specifying how to find the nearest key */
+enum class ENearestKeyOption : uint8
+{
+	NKO_None = 0x00,
+
+	/* Search keys */
+	NKO_SearchKeys = 0x01,
+
+	/* Search sections */
+	NKO_SearchSections = 0x02,
+
+	/* Search markers */
+	NKO_SearchMarkers = 0x04,
+
+	/** Search all tracks */
+	NKO_SearchAllTracks = 0x08
+};
+ENUM_CLASS_FLAGS(ENearestKeyOption);
+
+struct FTimeSliderArgs
+{
+	FTimeSliderArgs()
+		: ScrubPosition(0)
+		, ViewRange(FAnimatedRange(0.0f, 5.0f))
+		, ClampRange(FAnimatedRange(-FLT_MAX / 2.f, FLT_MAX / 2.f))
+		, AllowZoom(true)
+	{}
+
+	/** The scrub position */
+	TAttribute<FFrameTime> ScrubPosition;
+
+	/** The scrub position text */
+	TAttribute<FString> ScrubPositionText;
+
+	/** The parent sequence that the scrub position display text is relative to */
+	TAttribute<FMovieSceneSequenceID> ScrubPositionParent;
+
+	/** Called when the scrub position parent sequence is changed */
+	FOnScrubPositionParentChanged OnScrubPositionParentChanged;
+
+	/** Attribute for the parent sequence chain of the current sequence */
+	TAttribute<TArray<FMovieSceneSequenceID>> ScrubPositionParentChain;
+
+	/** View time range */
+	TAttribute< FAnimatedRange > ViewRange;
+
+	/** Clamp time range */
+	TAttribute< FAnimatedRange > ClampRange;
+
+	/** Called when the scrub position changes */
+	FOnScrubPositionChanged OnScrubPositionChanged;
+
+	/** Called right before the scrubber begins to move */
+	FSimpleDelegate OnBeginScrubberMovement;
+
+	/** Called right after the scrubber handle is released by the user */
+	FSimpleDelegate OnEndScrubberMovement;
+
+	/** Called when the view range changes */
+	FOnViewRangeChanged OnViewRangeChanged;
+
+	/** Called when the clamp range changes */
+	FOnTimeRangeChanged OnClampRangeChanged;
+
+	/** Delegate that is called when getting the nearest key */
+	FOnGetNearestKey OnGetNearestKey;
+
+	/** Attribute defining the active sub-sequence range for this controller */
+	TAttribute<TOptional<TRange<FFrameNumber>>> SubSequenceRange;
+
+	/** Attribute defining the playback range for this controller */
+	TAttribute<TRange<FFrameNumber>> PlaybackRange;
+
+	/** Attribute for the current sequence's display rate */
+	TAttribute<FFrameRate> DisplayRate;
+
+	/** Attribute for the current sequence's tick resolution */
+	TAttribute<FFrameRate> TickResolution;
+
+	/** Delegate that is called when the playback range wants to change */
+	FOnFrameRangeChanged OnPlaybackRangeChanged;
+
+	/** Called right before the playback range starts to be dragged */
+	FSimpleDelegate OnPlaybackRangeBeginDrag;
+
+	/** Called right after the playback range has finished being dragged */
+	FSimpleDelegate OnPlaybackRangeEndDrag;
+
+	/** Attribute defining the selection range for this controller */
+	TAttribute<TRange<FFrameNumber>> SelectionRange;
+
+	/** Delegate that is called when the selection range wants to change */
+	FOnFrameRangeChanged OnSelectionRangeChanged;
+
+	/** Called right before the selection range starts to be dragged */
+	FSimpleDelegate OnSelectionRangeBeginDrag;
+
+	/** Called right after the selection range has finished being dragged */
+	FSimpleDelegate OnSelectionRangeEndDrag;
+
+	/** Called right before a mark starts to be dragged */
+	FSimpleDelegate OnMarkBeginDrag;
+
+	/** Called right after a mark has finished being dragged */
+	FSimpleDelegate OnMarkEndDrag;
+
+	/** Attribute for the current sequence's vertical frames */
+	TAttribute<TSet<FFrameNumber>> VerticalFrames;
+
+	/** Attribute for the current sequence's marked frames */
+	TAttribute<TArray<FMovieSceneMarkedFrame>> MarkedFrames;
+
+	/** Attribute for the marked frames that might need to be shown, but do not belong to the current sequence*/
+	TAttribute<TArray<FMovieSceneMarkedFrame>> GlobalMarkedFrames;
+
+	/** Called when the marked frame needs to be set */
+	FOnSetMarkedFrame OnSetMarkedFrame;
+
+	/** Called when a marked frame is added */
+	FOnAddMarkedFrame OnAddMarkedFrame;
+
+	/** Called when a marked frame is deleted */
+	FOnDeleteMarkedFrame OnDeleteMarkedFrame;
+
+	/** Called when all marked frames should be deleted */
+	FSimpleDelegate OnDeleteAllMarkedFrames;
+
+	/** Round the scrub position to an integer during playback */
+	TAttribute<EMovieScenePlayerStatus::Type> PlaybackStatus;
+
+	/** Attribute defining whether the playback range is locked */
+	TAttribute<bool> IsPlaybackRangeLocked;
+
+	/** Attribute defining the time snap interval */
+	TAttribute<float> TimeSnapInterval;
+
+	/** Called when toggling the playback range lock */
+	FSimpleDelegate OnTogglePlaybackRangeLocked;
+
+	/** If we are allowed to zoom */
+	bool AllowZoom;
+
+	/** Numeric Type interface for converting between frame numbers and display formats. */
+	TSharedPtr<INumericTypeInterface<double>> NumericTypeInterface;
+};
diff --git a/Engine/Source/Editor/SequencerWidgets/Public/ViewRangeInterpolation.h b/Engine/Source/Editor/SequencerWidgets/Public/ViewRangeInterpolation.h
new file mode 100644
--- /dev/null
+++ b/Engine/Source/Editor/SequencerWidgets/Public/ViewRangeInterpolation.h
@@ -0,0 +1,12 @@
+// Copyright Epic Games, Inc. All Rights Reserved.
+
+#pragma once
+
+/** Enum specifying how to interpolate to a new view range */
+enum class EViewRangeInterpolation
+{
+	/** Use an externally defined animated interpolation */
+	Animated,
+	/** Set the view range immediately */
+	Immediate,
+};
diff --git a/Engine/Source/Editor/VREditor/UI/VREditorUISystem.cpp b/Engine/Source/Editor/VREditor/UI/VREditorUISystem.cpp
--- a/Engine/Source/Editor/VREditor/UI/VREditorUISystem.cpp
+++ b/Engine/Source/Editor/VREditor/UI/VREditorUISystem.cpp
@@ -81,6 +81,7 @@
 #include "Components/StaticMeshComponent.h"
 
 #include "ISequencer.h"
+#include "ITimeSlider.h"
 #include "Engine/StaticMesh.h"
 #include "EdMode.h"
 #include "Widgets/Images/SImage.h"
diff --git a/Engine/Source/Runtime/MovieScene/Private/MovieScene.cpp b/Engine/Source/Runtime/MovieScene/Private/MovieScene.cpp
--- a/Engine/Source/Runtime/MovieScene/Private/MovieScene.cpp
+++ b/Engine/Source/Runtime/MovieScene/Private/MovieScene.cpp
@@ -17,6 +17,7 @@
 #include "UObject/UObjectHash.h"
 
 #include UE_INLINE_GENERATED_CPP_BY_NAME(MovieScene)
+#include UE_INLINE_GENERATED_CPP_BY_NAME(MovieSceneMarkedFrame)
 
 #define LOCTEXT_NAMESPACE "MovieScene"
 
diff --git a/Engine/Source/Runtime/MovieScene/Public/MovieScene.h b/Engine/Source/Runtime/MovieScene/Public/MovieScene.h
--- a/Engine/Source/Runtime/MovieScene/Public/MovieScene.h
+++ b/Engine/Source/Runtime/MovieScene/Public/MovieScene.h
@@ -21,6 +21,7 @@
 #include "MovieSceneBinding.h"
 #include "MovieSceneFrameMigration.h"
 #include "MovieSceneFwd.h"
+#include "MovieSceneMarkedFrame.h"
 #include "MovieSceneObjectBindingID.h"
 #include "MovieScenePossessable.h"
 #include "MovieSceneSequenceID.h"
@@ -72,48 +73,6 @@ struct FMovieSceneExpansionState
 	bool bExpanded;
 };
 
-USTRUCT(BlueprintType)
-struct FMovieSceneMarkedFrame
-{
-	GENERATED_BODY()
-
-	FMovieSceneMarkedFrame()
-		: Label(FString())
-#if WITH_EDITORONLY_DATA
-		, Comment(FString())
-		, Color(0.f, 1.f, 1.f, 0.4f)
-#endif
-		, bIsDeterminismFence(false)
-	{}
-
-	FMovieSceneMarkedFrame(FFrameNumber InFrameNumber)
-		: FrameNumber(InFrameNumber)
-		, Label(FString())
-#if WITH_EDITORONLY_DATA
-		, Comment(FString())
-		, Color(0.f, 1.f, 1.f, 0.4f)
-#endif
-		, bIsDeterminismFence(false)
-	{}
-
-	UPROPERTY(EditAnywhere, BlueprintReadWrite, Category = "Marked Frame")
-	FFrameNumber FrameNumber;
-
-	UPROPERTY(EditAnywhere, BlueprintReadWrite, Category = "Marked Frame")
-	FString Label;
-
-#if WITH_EDITORONLY_DATA
-	UPROPERTY(EditAnywhere, Category = "Marked Frame")
-	FString Comment;
-
-	UPROPERTY(EditAnywhere, Category = "Marked Frame")
-	FLinearColor Color;
-#endif
-
-	UPROPERTY(EditAnywhere, BlueprintReadWrite, Category = "Marked Frame", DisplayName="Is Determinism Fence?", meta=(Tooltip="When checked, treat this mark as a fence for evaluation purposes. Fences cannot be crossed in a single evaluation, and force the evaluation to be split into 2 separate parts."))
-	bool bIsDeterminismFence;
-};
-
 /**
  * Editor only data that needs to be saved between sessions for editing but has no runtime purpose
  */
diff --git a/Engine/Source/Runtime/MovieScene/Public/MovieSceneMarkedFrame.h b/Engine/Source/Runtime/MovieScene/Public/MovieSceneMarkedFrame.h
new file mode 100644
--- /dev/null
+++ b/Engine/Source/Runtime/MovieScene/Public/MovieSceneMarkedFrame.h
@@ -0,0 +1,46 @@
+// Copyright Epic Games, Inc. All Rights Reserved.
+
+#pragma once
+
+#include "Containers/UnrealString.h"
+#include "Misc/FrameNumber.h"
+#include "Math/Color.h"
+
+#include "MovieSceneMarkedFrame.generated.h"
+
+USTRUCT(BlueprintType)
+struct FMovieSceneMarkedFrame
+{
+	GENERATED_BODY()
+
+	FMovieSceneMarkedFrame()
+		: Label(FString())
+#if WITH_EDITORONLY_DATA
+		, Color(0.f, 1.f, 1.f, 0.4f)
+#endif
+		, bIsDeterminismFence(false)
+	{}
+
+	FMovieSceneMarkedFrame(FFrameNumber InFrameNumber)
+		: FrameNumber(InFrameNumber)
+		, Label(FString())
+#if WITH_EDITORONLY_DATA
+		, Color(0.f, 1.f, 1.f, 0.4f)
+#endif
+		, bIsDeterminismFence(false)
+	{}
+
+	UPROPERTY(EditAnywhere, BlueprintReadWrite, Category = "Marked Frame")
+	FFrameNumber FrameNumber;
+
+	UPROPERTY(EditAnywhere, BlueprintReadWrite, Category = "Marked Frame")
+	FString Label;
+
+#if WITH_EDITORONLY_DATA
+	UPROPERTY(EditAnywhere, Category = "Marked Frame")
+	FLinearColor Color;
+#endif
+
+	UPROPERTY(EditAnywhere, BlueprintReadWrite, Category = "Marked Frame", DisplayName = "Is Determinism Fence?", meta = (Tooltip = "When checked, treat this mark as a fence for evaluation purposes. Fences cannot be crossed in a single evaluation, and force the evaluation to be split into 2 separate parts."))
+	bool bIsDeterminismFence;
+};
