Animation mode & Sequencer: added the ability to filter widget types to avoid stealing focus from

#rnx

diff --git a/Engine/Plugins/Animation/ControlRig/Source/ControlRigEditor/Private/EditMode/ControlRigEditMode.cpp b/Engine/Plugins/Animation/ControlRig/Source/ControlRigEditor/Private/EditMode/ControlRigEditMode.cpp
--- a/Engine/Plugins/Animation/ControlRig/Source/ControlRigEditor/Private/EditMode/ControlRigEditMode.cpp
+++ b/Engine/Plugins/Animation/ControlRig/Source/ControlRigEditor/Private/EditMode/ControlRigEditMode.cpp
@@ -175,7 +175,8 @@ namespace ControlRigSelectionConstants
 }
 
 FControlRigEditMode::FControlRigEditMode()
-	: bIsChangingControlShapeTransform(false)
+	: PendingFocus(FPendingWidgetFocus::MakeNoTextEdit())
+	, bIsChangingControlShapeTransform(false)
 	, bIsTracking(false)
 	, bManipulatorMadeChange(false)
 	, bSelecting(false)
diff --git a/Engine/Source/Editor/CurveEditor/Private/SCurveEditorPanel.cpp b/Engine/Source/Editor/CurveEditor/Private/SCurveEditorPanel.cpp
--- a/Engine/Source/Editor/CurveEditor/Private/SCurveEditorPanel.cpp
+++ b/Engine/Source/Editor/CurveEditor/Private/SCurveEditorPanel.cpp
@@ -123,7 +123,8 @@ struct FCurveEditorPanelViewTracker
 };
 
 SCurveEditorPanel::SCurveEditorPanel()
-	: bNeedsRefresh(true)
+	: PendingFocus(FPendingWidgetFocus::MakeNoTextEdit())
+	, bNeedsRefresh(true)
 	, CachedActiveCurvesSerialNumber(-1)
 {
 	EditObjects = MakeUnique<FCurveEditorEditObjectContainer>();
diff --git a/Engine/Source/Editor/Sequencer/Private/SSequencer.cpp b/Engine/Source/Editor/Sequencer/Private/SSequencer.cpp
--- a/Engine/Source/Editor/Sequencer/Private/SSequencer.cpp
+++ b/Engine/Source/Editor/Sequencer/Private/SSequencer.cpp
@@ -165,6 +165,7 @@ void SSequencer::Construct(const FArguments& InArgs, TSharedRef<FSequencer> InSe
 	bUserIsSelecting = false;
 	CachedClampRange = TRange<double>::Empty();
 	CachedViewRange = TRange<double>::Empty();
+	PendingFocus = FPendingWidgetFocus::MakeNoTextEdit();
 
 	OnPlaybackRangeBeginDrag = InArgs._OnPlaybackRangeBeginDrag;
 	OnPlaybackRangeEndDrag = InArgs._OnPlaybackRangeEndDrag;
diff --git a/Engine/Source/Editor/SequencerWidgets/Private/WidgetFocusUtils.cpp b/Engine/Source/Editor/SequencerWidgets/Private/WidgetFocusUtils.cpp
--- a/Engine/Source/Editor/SequencerWidgets/Private/WidgetFocusUtils.cpp
+++ b/Engine/Source/Editor/SequencerWidgets/Private/WidgetFocusUtils.cpp
@@ -5,6 +5,19 @@
 #include "Framework/Application/SlateApplication.h"
 #include "Framework/Application/SlateUser.h"
 
+FPendingWidgetFocus::FPendingWidgetFocus(const TArray<FName>& InTypesKeepingFocus)
+	: KeepingFocus(InTypesKeepingFocus)
+{}
+
+FPendingWidgetFocus FPendingWidgetFocus::MakeNoTextEdit()
+{
+	static TArray<FName> EditableTextTypes({"SEditableText"});
+	// NOTE: "SMultiLineEditableText" might be added as well
+	
+	static FPendingWidgetFocus NewPendingFocus(EditableTextTypes);
+	return NewPendingFocus;
+}
+
 FPendingWidgetFocus::~FPendingWidgetFocus()
 {
 	PendingFocusFunction.Reset();
@@ -23,6 +36,12 @@ void FPendingWidgetFocus::SetPendingFocusIfNeeded(const TWeakPtr<SWidget>& InWid
 		return;
 	}
 
+	if (!CanFocusBeStolen())
+	{
+		PendingFocusFunction.Reset();
+		return;
+	}
+	
 	PendingFocusFunction = [WidgetFocus = InWidget]()
 	{
 		if (WidgetFocus.IsValid())
@@ -71,3 +90,21 @@ void FPendingWidgetFocus::OnPreInputKeyDown(const FKeyEvent&)
 		PendingFocusFunction.Reset();
 	}
 }
+
+bool FPendingWidgetFocus::CanFocusBeStolen() const
+{
+	if (!KeepingFocus.IsEmpty())
+	{
+		bool bShouldCurrentFocusBeKept = false;
+		FSlateApplication::Get().ForEachUser([this, &bShouldCurrentFocusBeKept](const FSlateUser& User)
+		{
+			if (TSharedPtr<SWidget> FocusedWidget = User.GetFocusedWidget())
+			{
+				bShouldCurrentFocusBeKept = KeepingFocus.Contains(FocusedWidget->GetType());
+			}
+		});
+		return !bShouldCurrentFocusBeKept;
+	}
+	return true;
+}
+
diff --git a/Engine/Source/Editor/SequencerWidgets/Public/WidgetFocusUtils.h b/Engine/Source/Editor/SequencerWidgets/Public/WidgetFocusUtils.h
--- a/Engine/Source/Editor/SequencerWidgets/Public/WidgetFocusUtils.h
+++ b/Engine/Source/Editor/SequencerWidgets/Public/WidgetFocusUtils.h
@@ -19,6 +19,11 @@ struct FKeyEvent;
 class SEQUENCERWIDGETS_API FPendingWidgetFocus : public TSharedFromThis<FPendingWidgetFocus>
 {
 public:
+
+	FPendingWidgetFocus() = default;
+	FPendingWidgetFocus(const TArray<FName>& InTypesKeepingFocus);
+
+	static FPendingWidgetFocus MakeNoTextEdit();
 	
 	~FPendingWidgetFocus();
 
@@ -31,7 +36,10 @@ public:
 private:
 	
 	void OnPreInputKeyDown(const FKeyEvent&);
+	bool CanFocusBeStolen() const;
 
 	TFunction<void()> PendingFocusFunction;
 	FDelegateHandle PreInputKeyDownHandle;
+
+	TArray<FName> KeepingFocus;
 };
