Animation mode & Sequencer: reset the pending function when a mouse button is pressed before any button

#rnx

diff --git a/Engine/Source/Editor/SequencerWidgets/Private/WidgetFocusUtils.cpp b/Engine/Source/Editor/SequencerWidgets/Private/WidgetFocusUtils.cpp
--- a/Engine/Source/Editor/SequencerWidgets/Private/WidgetFocusUtils.cpp
+++ b/Engine/Source/Editor/SequencerWidgets/Private/WidgetFocusUtils.cpp
@@ -21,17 +21,23 @@ FPendingWidgetFocus FPendingWidgetFocus::MakeNoTextEdit()
 FPendingWidgetFocus::~FPendingWidgetFocus()
 {
 	PendingFocusFunction.Reset();
+	
+	FSlateApplication& SlateApplication = FSlateApplication::Get();
 	if (PreInputKeyDownHandle.IsValid())
 	{
-		FSlateApplication& SlateApplication = FSlateApplication::Get();
 		SlateApplication.OnApplicationPreInputKeyDownListener().Remove(PreInputKeyDownHandle);
 		PreInputKeyDownHandle.Reset();
 	}
+	if (PreInputButtonDownHandle.IsValid())
+	{
+		SlateApplication.OnApplicationMousePreInputButtonDownListener().Remove(PreInputButtonDownHandle);
+		PreInputButtonDownHandle.Reset();
+	}
 }
 
 void FPendingWidgetFocus::SetPendingFocusIfNeeded(const TWeakPtr<SWidget>& InWidget)
 {
-	if (!PreInputKeyDownHandle.IsValid())
+	if (!IsEnabled())
 	{
 		return;
 	}
@@ -68,18 +74,24 @@ void FPendingWidgetFocus::Enable(const bool InEnabled)
 		SlateApplication.OnApplicationPreInputKeyDownListener().Remove(PreInputKeyDownHandle);
 		PreInputKeyDownHandle.Reset();
 	}
+	if (PreInputButtonDownHandle.IsValid())
+	{
+		SlateApplication.OnApplicationMousePreInputButtonDownListener().Remove(PreInputButtonDownHandle);
+		PreInputButtonDownHandle.Reset();
+	}
 	
 	PendingFocusFunction.Reset();
 	
 	if (InEnabled)
 	{
 		PreInputKeyDownHandle = SlateApplication.OnApplicationPreInputKeyDownListener().AddRaw(this, &FPendingWidgetFocus::OnPreInputKeyDown);
+		PreInputButtonDownHandle = SlateApplication.OnApplicationMousePreInputButtonDownListener().AddRaw(this, &FPendingWidgetFocus::OnPreInputButtonDown);
 	}
 }
 
 bool FPendingWidgetFocus::IsEnabled() const
 {
-	return PreInputKeyDownHandle.IsValid();
+	return PreInputKeyDownHandle.IsValid() && PreInputButtonDownHandle.IsValid();
 }
 	
 void FPendingWidgetFocus::OnPreInputKeyDown(const FKeyEvent&)
@@ -91,6 +103,13 @@ void FPendingWidgetFocus::OnPreInputKeyDown(const FKeyEvent&)
 	}
 }
 
+void FPendingWidgetFocus::OnPreInputButtonDown(const FPointerEvent&)
+{
+	// remove any pending focus as clicking a mouse button will set the focus
+	// so this pending function should not interfere.
+	PendingFocusFunction.Reset();
+}
+
 bool FPendingWidgetFocus::CanFocusBeStolen() const
 {
 	if (!KeepingFocus.IsEmpty())
diff --git a/Engine/Source/Editor/SequencerWidgets/Public/WidgetFocusUtils.h b/Engine/Source/Editor/SequencerWidgets/Public/WidgetFocusUtils.h
--- a/Engine/Source/Editor/SequencerWidgets/Public/WidgetFocusUtils.h
+++ b/Engine/Source/Editor/SequencerWidgets/Public/WidgetFocusUtils.h
@@ -8,6 +8,7 @@
 
 class SWidget;
 struct FKeyEvent;
+struct FPointerEvent;
 
 /**
  * FPendingWidgetFocus is a utility class that stores a pending focus function when a SWidget is hovered over.
@@ -36,10 +37,12 @@ public:
 private:
 	
 	void OnPreInputKeyDown(const FKeyEvent&);
+	void OnPreInputButtonDown(const FPointerEvent&);
 	bool CanFocusBeStolen() const;
 
 	TFunction<void()> PendingFocusFunction;
 	FDelegateHandle PreInputKeyDownHandle;
+	FDelegateHandle PreInputButtonDownHandle;
 
 	TArray<FName> KeepingFocus;
 };
