Animation mode & Sequencer: fix ensure when closing UE (follow-up to 36329357)

diff --git a/Engine/Source/Editor/SequencerWidgets/Private/WidgetFocusUtils.cpp b/Engine/Source/Editor/SequencerWidgets/Private/WidgetFocusUtils.cpp
--- a/Engine/Source/Editor/SequencerWidgets/Private/WidgetFocusUtils.cpp
+++ b/Engine/Source/Editor/SequencerWidgets/Private/WidgetFocusUtils.cpp
@@ -21,22 +21,31 @@ FPendingWidgetFocus FPendingWidgetFocus::MakeNoTextEdit()
 FPendingWidgetFocus::~FPendingWidgetFocus()
 {
 	PendingFocusFunction.Reset();
-	
-	FSlateApplication& SlateApplication = FSlateApplication::Get();
-	if (PreInputKeyDownHandle.IsValid())
-	{
-		SlateApplication.OnApplicationPreInputKeyDownListener().Remove(PreInputKeyDownHandle);
-		PreInputKeyDownHandle.Reset();
-	}
-	if (PreInputButtonDownHandle.IsValid())
-	{
-		SlateApplication.OnApplicationMousePreInputButtonDownListener().Remove(PreInputButtonDownHandle);
-		PreInputButtonDownHandle.Reset();
+
+	if (FSlateApplication::IsInitialized())
+	{	
+		FSlateApplication& SlateApplication = FSlateApplication::Get();
+		if (PreInputKeyDownHandle.IsValid())
+		{
+			SlateApplication.OnApplicationPreInputKeyDownListener().Remove(PreInputKeyDownHandle);
+		}
+		if (PreInputButtonDownHandle.IsValid())
+		{
+			SlateApplication.OnApplicationMousePreInputButtonDownListener().Remove(PreInputButtonDownHandle);
+		}
 	}
+
+	PreInputKeyDownHandle.Reset();
+	PreInputButtonDownHandle.Reset();
 }
 
 void FPendingWidgetFocus::SetPendingFocusIfNeeded(const TWeakPtr<SWidget>& InWidget)
 {
+	if (!FSlateApplication::IsInitialized())
+	{
+		return;	
+	}
+	
 	if (!IsEnabled())
 	{
 		return;
@@ -68,6 +77,14 @@ void FPendingWidgetFocus::ResetPendingFocus()
 
 void FPendingWidgetFocus::Enable(const bool InEnabled)
 {
+	if (!FSlateApplication::IsInitialized())
+	{
+		PendingFocusFunction.Reset();
+		PreInputKeyDownHandle.Reset();
+		PreInputButtonDownHandle.Reset();
+		return;	
+	}
+	
 	FSlateApplication& SlateApplication = FSlateApplication::Get();
 	if (PreInputKeyDownHandle.IsValid())
 	{
@@ -115,13 +132,18 @@ bool FPendingWidgetFocus::CanFocusBeStolen() const
 	if (!KeepingFocus.IsEmpty())
 	{
 		bool bShouldCurrentFocusBeKept = false;
-		FSlateApplication::Get().ForEachUser([this, &bShouldCurrentFocusBeKept](const FSlateUser& User)
+
+		if (FSlateApplication::IsInitialized())
 		{
-			if (TSharedPtr<SWidget> FocusedWidget = User.GetFocusedWidget())
-			{
-				bShouldCurrentFocusBeKept = KeepingFocus.Contains(FocusedWidget->GetType());
-			}
-		});
+			FSlateApplication::Get().ForEachUser([this, &bShouldCurrentFocusBeKept](const FSlateUser& User)
+		   {
+			   if (TSharedPtr<SWidget> FocusedWidget = User.GetFocusedWidget())
+			   {
+				   bShouldCurrentFocusBeKept = KeepingFocus.Contains(FocusedWidget->GetType());
+			   }
+		   });
+		}
+		
 		return !bShouldCurrentFocusBeKept;
 	}
 	return true;
