Sequencer: Fix computing drag position when the desktop has a display scale

#Synced-CL 48032220

diff --git a/Engine/Source/Editor/SequencerWidgets/Private/STimeRangeSlider.cpp b/Engine/Source/Editor/SequencerWidgets/Private/STimeRangeSlider.cpp
--- a/Engine/Source/Editor/SequencerWidgets/Private/STimeRangeSlider.cpp
+++ b/Engine/Source/Editor/SequencerWidgets/Private/STimeRangeSlider.cpp
@@ -47,7 +47,8 @@ double STimeRangeSlider::ComputeDragDelta(const FPointerEvent& MouseEvent, doubl
 	}
 	double DragDistance = (MouseEvent.GetScreenSpacePosition() - MouseDownPosition).X;
 
-	const double PixelToUnits = (EndTime - StartTime) / (GeometryWidth - TimeRangeSliderConstants::HandleSize*2);
+	const double Width = (GeometryWidth - TimeRangeSliderConstants::HandleSize * 2);
+	const double PixelToUnits = FMath::IsNearlyZero(Width) ? 0 : (EndTime - StartTime) / Width;
 	return DragDistance * PixelToUnits;
 }
 
@@ -176,7 +177,7 @@ FReply STimeRangeSlider::OnMouseMove( const FGeometry& MyGeometry, const FPointe
 {
 	if (HasMouseCapture())
 	{
-		double DragDelta = ComputeDragDelta(MouseEvent, MyGeometry.GetLocalSize().X);
+		double DragDelta = ComputeDragDelta(MouseEvent, MyGeometry.GetAbsoluteSize().X);
 
 		ITimeSliderController* TimeSliderControllerPtr = TimeSliderController.Get();
 		if (!TimeSliderControllerPtr)
